Initial implementation of OnlineServicesEos
Lots of temporary placeholder things to be fleshed out more and improved on later.

diff --git a/Engine/Plugins/Online/Experimental/OnlineServices/Source/OnlineServicesCommon/Private/Online/AuthCommon.cpp b/Engine/Plugins/Online/Experimental/OnlineServices/Source/OnlineServicesCommon/Private/Online/AuthCommon.cpp
--- a/Engine/Plugins/Online/Experimental/OnlineServices/Source/OnlineServicesCommon/Private/Online/AuthCommon.cpp
+++ b/Engine/Plugins/Online/Experimental/OnlineServices/Source/OnlineServicesCommon/Private/Online/AuthCommon.cpp
@@ -34,6 +34,16 @@ TOnlineAsyncOpHandle<FAuthGenerateAuth> FAuthCommon::GenerateAuth(FAuthGenerateA
 	return Operation.GetHandle();
 }
 
+TOnlineResult<FAuthGetAccountByLocalUserNum::Result> FAuthCommon::GetAccountByLocalUserNum(FAuthGetAccountByLocalUserNum::Params&& Params)
+{
+	return TOnlineResult<FAuthGetAccountByLocalUserNum::Result>(Errors::Unimplemented());
+}
+
+TOnlineResult<FAuthGetAccountByAccountId::Result> FAuthCommon::GetAccountByAccountId(FAuthGetAccountByAccountId::Params&& Params)
+{
+	return TOnlineResult<FAuthGetAccountByAccountId::Result>(Errors::Unimplemented());
+}
+
 TOnlineEvent<void(const FLoginStatusChanged&)> FAuthCommon::OnLoginStatusChanged()
 {
 	return OnLoginStatusChangedEvent;
diff --git a/Engine/Plugins/Online/Experimental/OnlineServices/Source/OnlineServicesCommon/Public/Online/AuthCommon.h b/Engine/Plugins/Online/Experimental/OnlineServices/Source/OnlineServicesCommon/Public/Online/AuthCommon.h
--- a/Engine/Plugins/Online/Experimental/OnlineServices/Source/OnlineServicesCommon/Public/Online/AuthCommon.h
+++ b/Engine/Plugins/Online/Experimental/OnlineServices/Source/OnlineServicesCommon/Public/Online/AuthCommon.h
@@ -19,6 +19,8 @@ public:
 	virtual TOnlineAsyncOpHandle<FAuthLogin> Login(FAuthLogin::Params&& Params) override;
 	virtual TOnlineAsyncOpHandle<FAuthLogout> Logout(FAuthLogout::Params&& Params) override;
 	virtual TOnlineAsyncOpHandle<FAuthGenerateAuth> GenerateAuth(FAuthGenerateAuth::Params&& Params) override;
+	virtual TOnlineResult<FAuthGetAccountByLocalUserNum::Result> GetAccountByLocalUserNum(FAuthGetAccountByLocalUserNum::Params&& Params) override;
+	virtual TOnlineResult<FAuthGetAccountByAccountId::Result> GetAccountByAccountId(FAuthGetAccountByAccountId::Params&& Params) override;
 	virtual TOnlineEvent<void(const FLoginStatusChanged&)> OnLoginStatusChanged() override;
 
 protected:
diff --git a/Engine/Plugins/Online/Experimental/OnlineServices/Source/OnlineServicesInterface/Public/Online/Auth.h b/Engine/Plugins/Online/Experimental/OnlineServices/Source/OnlineServicesInterface/Public/Online/Auth.h
--- a/Engine/Plugins/Online/Experimental/OnlineServices/Source/OnlineServicesInterface/Public/Online/Auth.h
+++ b/Engine/Plugins/Online/Experimental/OnlineServices/Source/OnlineServicesInterface/Public/Online/Auth.h
@@ -62,7 +62,7 @@ struct FAuthGenerateAuth
 	};
 };
 
-struct FGetAccountByLocalUserNum
+struct FAuthGetAccountByLocalUserNum
 {
 	struct Params
 	{
@@ -75,7 +75,7 @@ struct FGetAccountByLocalUserNum
 	};
 };
 
-struct FGetAccountByAccountId
+struct FAuthGetAccountByAccountId
 {
 	struct Params
 	{
@@ -113,12 +113,12 @@ public:
 	/**
 	 * Get logged in user by local user num
 	 */ 
-	virtual TOnlineResult<FGetAccountByLocalUserNum::Result> GetAccountByLocalUserNum(FGetAccountByLocalUserNum::Params&& Params) = 0;
+	virtual TOnlineResult<FAuthGetAccountByLocalUserNum::Result> GetAccountByLocalUserNum(FAuthGetAccountByLocalUserNum::Params&& Params) = 0;
 
 	/**
 	 * Get logged in user by account id
 	 */
-	virtual TOnlineResult<FGetAccountByAccountId::Result> GetAccountByAccountId(FGetAccountByAccountId::Params&& Params) = 0;
+	virtual TOnlineResult<FAuthGetAccountByAccountId::Result> GetAccountByAccountId(FAuthGetAccountByAccountId::Params&& Params) = 0;
 
 	/**
 	 * Event triggered when a local user's login status changes
diff --git a/Engine/Plugins/Online/Experimental/OnlineServices/Source/OnlineServicesInterface/Public/Online/AuthErrors.h b/Engine/Plugins/Online/Experimental/OnlineServices/Source/OnlineServicesInterface/Public/Online/AuthErrors.h
new file mode 100644
--- /dev/null
+++ b/Engine/Plugins/Online/Experimental/OnlineServices/Source/OnlineServicesInterface/Public/Online/AuthErrors.h
@@ -0,0 +1,12 @@
+// Copyright Epic Games, Inc. All Rights Reserved.
+
+#pragma once
+
+#include "Online/OnlineError.h"
+
+ // RobC mod- whole file
+namespace UE::Online::Errors {
+
+inline FOnlineError InvalidCredentials() { return FOnlineError(TEXT("invalid_credentials")); }
+
+/* UE::Online::Errors */ }
diff --git a/Engine/Plugins/Online/Experimental/OnlineServices/Source/OnlineServicesInterface/Public/Online/OnlineError.h b/Engine/Plugins/Online/Experimental/OnlineServices/Source/OnlineServicesInterface/Public/Online/OnlineError.h
--- a/Engine/Plugins/Online/Experimental/OnlineServices/Source/OnlineServicesInterface/Public/Online/OnlineError.h
+++ b/Engine/Plugins/Online/Experimental/OnlineServices/Source/OnlineServicesInterface/Public/Online/OnlineError.h
@@ -22,6 +22,7 @@ namespace Errors
 {
 	inline FOnlineError Cancelled() { return FOnlineError(); }
 	inline FOnlineError Unimplemented() { return FOnlineError(); }
+	inline FOnlineError UnknownError() { return FOnlineError(); }
 }
 
 /* UE::Online */ }
diff --git a/Engine/Plugins/Online/Experimental/OnlineServicesEOS/OnlineServicesEOS.uplugin b/Engine/Plugins/Online/Experimental/OnlineServicesEOS/OnlineServicesEOS.uplugin
new file mode 100644
--- /dev/null
+++ b/Engine/Plugins/Online/Experimental/OnlineServicesEOS/OnlineServicesEOS.uplugin
@@ -0,0 +1,32 @@
+{
+	"FileVersion": 3,
+
+	"FriendlyName": "Online Services EOS",
+	"Version": 1,
+	"VersionName": "1.0",
+	"Description": "Online Services implementation for EOS",
+	"Category": "Online Platform",
+	"CreatedBy": "Epic Games, Inc.",
+	"CreatedByURL": "http://epicgames.com",
+	"EnabledByDefault": false,
+
+	"Modules": [
+		{
+			"Name": "OnlineServicesEOS",
+			"Type": "Runtime",
+			"LoadingPhase": "PostConfigInit"
+		}
+	],
+	"LocalizationTargets": [
+		{
+			"Name": "OnlineSubsystemEOS",
+			"LoadingPolicy": "Always"
+		}
+	],
+	"Plugins": [
+		{
+			"Name": "OnlineServices",
+			"Enabled": true
+		}
+	]
+}
\ No newline at end of file
diff --git a/Engine/Plugins/Online/Experimental/OnlineServicesEOS/Source/OnlineServicesEOS.Build.cs b/Engine/Plugins/Online/Experimental/OnlineServicesEOS/Source/OnlineServicesEOS.Build.cs
new file mode 100644
--- /dev/null
+++ b/Engine/Plugins/Online/Experimental/OnlineServicesEOS/Source/OnlineServicesEOS.Build.cs
@@ -0,0 +1,21 @@
+// Copyright Epic Games, Inc. All Rights Reserved.
+
+using UnrealBuildTool;
+using System.IO;
+
+public class OnlineServicesEOS : ModuleRules
+{
+	public OnlineServicesEOS(ReadOnlyTargetRules Target) : base(Target)
+    {
+		PublicDependencyModuleNames.AddRange(
+			new string[] {
+				"Core",
+				"OnlineServicesInterface",
+				"OnlineServicesCommon",
+			}
+		);
+
+		// TODO:  Use EOSShared module
+		PrivateDefinitions.Add("WITH_EOS_SDK=0");
+	}
+}
diff --git a/Engine/Plugins/Online/Experimental/OnlineServicesEOS/Source/Private/AuthEOS.cpp b/Engine/Plugins/Online/Experimental/OnlineServicesEOS/Source/Private/AuthEOS.cpp
new file mode 100644
--- /dev/null
+++ b/Engine/Plugins/Online/Experimental/OnlineServicesEOS/Source/Private/AuthEOS.cpp
@@ -0,0 +1,471 @@
+// Copyright Epic Games, Inc. All Rights Reserved.
+
+#include "AuthEOS.h"
+#include "OnlineServicesEOS.h"
+#include "OnlineServicesEOSTypes.h"
+#include "Online/AuthErrors.h"
+
+#if WITH_EOS_SDK
+#include "eos_auth.h"
+#include "eos_types.h"
+#include "eos_init.h"
+#include "eos_sdk.h"
+#include "eos_logging.h"
+
+void EOS_CALL EOSFree(void* Ptr)
+{
+	FMemory::Free(Ptr);
+}
+
+void* EOS_CALL EOSMalloc(size_t Size, size_t Alignment)
+{
+	return FMemory::Malloc(Size, Alignment);
+}
+
+void* EOS_CALL EOSRealloc(void* Ptr, size_t Size, size_t Alignment)
+{
+	return FMemory::Realloc(Ptr, Size, Alignment);
+}
+
+void EOS_CALL EOSLog(const EOS_LogMessage* InMsg)
+{
+	if (GLog == nullptr)
+	{
+		return;
+	}
+
+	switch (InMsg->Level)
+	{
+		case EOS_ELogLevel::EOS_LOG_Fatal:
+		{
+			UE_LOG(LogTemp, Warning, TEXT("EOSSDK-%s: %s"), ANSI_TO_TCHAR(InMsg->Category), ANSI_TO_TCHAR(InMsg->Message));
+			break;
+		}
+		case EOS_ELogLevel::EOS_LOG_Error:
+		{
+			UE_LOG(LogTemp, Warning, TEXT("EOSSDK-%s: %s"), ANSI_TO_TCHAR(InMsg->Category), ANSI_TO_TCHAR(InMsg->Message));
+			break;
+		}
+		case EOS_ELogLevel::EOS_LOG_Warning:
+		{
+			UE_LOG(LogTemp, Warning, TEXT("EOSSDK-%s: %s"), ANSI_TO_TCHAR(InMsg->Category), ANSI_TO_TCHAR(InMsg->Message));
+			break;
+		}
+		case EOS_ELogLevel::EOS_LOG_Verbose:
+		{
+			UE_LOG(LogTemp, Warning, TEXT("EOSSDK-%s: %s"), ANSI_TO_TCHAR(InMsg->Category), ANSI_TO_TCHAR(InMsg->Message));
+			break;
+		}
+		case EOS_ELogLevel::EOS_LOG_VeryVerbose:
+		{
+			UE_LOG(LogTemp, Warning, TEXT("EOSSDK-%s: %s"), ANSI_TO_TCHAR(InMsg->Category), ANSI_TO_TCHAR(InMsg->Message));
+			break;
+		}
+		case EOS_ELogLevel::EOS_LOG_Info:
+		default:
+		{
+			UE_LOG(LogTemp, Warning, TEXT("EOSSDK-%s: %s"), ANSI_TO_TCHAR(InMsg->Category), ANSI_TO_TCHAR(InMsg->Message));
+			break;
+		}
+	}
+}
+
+EOS_HPlatform SetupEOSSdk()
+{
+	static EOS_HPlatform PlatformHandle = nullptr;
+	if (PlatformHandle == nullptr)
+	{
+		// Init EOS SDK
+		EOS_InitializeOptions SDKOptions = { };
+		SDKOptions.ApiVersion = EOS_INITIALIZE_API_LATEST;
+
+		SDKOptions.ProductName = "My Product";
+		SDKOptions.ProductVersion = "Unknown";
+		SDKOptions.AllocateMemoryFunction = &EOSMalloc;
+		SDKOptions.ReallocateMemoryFunction = &EOSRealloc;
+		SDKOptions.ReleaseMemoryFunction = &EOSFree;
+
+		EOS_EResult InitResult = EOS_Initialize(&SDKOptions);
+
+		EOS_Platform_Options PlatformOptions = {};
+		PlatformOptions.ApiVersion = EOS_PLATFORM_OPTIONS_API_LATEST;
+		PlatformOptions.bIsServer = EOS_FALSE;
+		PlatformOptions.EncryptionKey = nullptr;
+		PlatformOptions.OverrideCountryCode = nullptr;
+		PlatformOptions.OverrideLocaleCode = nullptr;
+		PlatformOptions.Flags = EOS_PF_WINDOWS_ENABLE_OVERLAY_D3D9 | EOS_PF_WINDOWS_ENABLE_OVERLAY_D3D10 | EOS_PF_WINDOWS_ENABLE_OVERLAY_OPENGL; // Enable overlay support for D3D9/10 and OpenGL. This sample uses D3D11 or SDL.
+
+		char CacheDirectory[512];
+		FCStringAnsi::Strncpy(CacheDirectory, TCHAR_TO_UTF8(*(FPlatformProcess::UserDir() / FString(TEXT("CacheDir")))), 512);
+		PlatformOptions.CacheDirectory = CacheDirectory;
+
+		PlatformOptions.ProductId = "";
+		PlatformOptions.SandboxId = "";
+		PlatformOptions.DeploymentId = "";
+
+		PlatformOptions.ClientCredentials.ClientId = "";
+		PlatformOptions.ClientCredentials.ClientSecret = "";
+
+		PlatformHandle = EOS_Platform_Create(&PlatformOptions);
+
+		EOS_Logging_SetCallback(&EOSLog);
+	}
+	return PlatformHandle;
+}
+
+namespace UE::Online {
+
+// TEMP until Net Id Registry is done
+TMap<EOS_EpicAccountId, int32> EOSAccountIdMap;
+
+// Copied from OSS EOS
+
+#define EOS_OSS_STRING_BUFFER_LENGTH 256
+// Chose arbitrarily since the SDK doesn't define it
+#define EOS_MAX_TOKEN_SIZE 4096
+
+struct FEosAuthCredentials :
+	public EOS_Auth_Credentials
+{
+	FEosAuthCredentials() :
+		EOS_Auth_Credentials()
+	{
+		ApiVersion = EOS_AUTH_CREDENTIALS_API_LATEST;
+		Id = IdAnsi;
+		Token = TokenAnsi;
+	}
+
+	FEosAuthCredentials(EOS_EExternalCredentialType InExternalType, const TArray<uint8>& InToken) :
+		EOS_Auth_Credentials()
+	{
+		ApiVersion = EOS_AUTH_CREDENTIALS_API_LATEST;
+		Type = EOS_ELoginCredentialType::EOS_LCT_ExternalAuth;
+		ExternalType = InExternalType;
+		Id = IdAnsi;
+		Token = TokenAnsi;
+
+		uint32_t InOutBufferLength = EOS_OSS_STRING_BUFFER_LENGTH;
+		EOS_ByteArray_ToString(InToken.GetData(), InToken.Num(), TokenAnsi, &InOutBufferLength);
+	}
+
+	char IdAnsi[EOS_OSS_STRING_BUFFER_LENGTH];
+	char TokenAnsi[EOS_MAX_TOKEN_SIZE];
+};
+
+FAuthEOS::FAuthEOS(FOnlineServicesEOS& InOwningSubsystem)
+	: FAuthCommon(InOwningSubsystem)
+{
+}
+
+void FAuthEOS::Tick(float DeltaTime)
+{
+	EOS_HPlatform PlatformHandle = SetupEOSSdk();
+	if (PlatformHandle)
+	{
+		EOS_Platform_Tick(PlatformHandle);
+	}
+}
+
+bool LexFromString(EOS_EAuthScopeFlags& OutEnum, const TCHAR* InString)
+{
+	if (FCString::Stricmp(InString, TEXT("BasicProfile")) == 0)
+	{
+		OutEnum = EOS_EAuthScopeFlags::EOS_AS_BasicProfile;
+	}
+	else if (FCString::Stricmp(InString, TEXT("FriendsList")) == 0)
+	{
+		OutEnum = EOS_EAuthScopeFlags::EOS_AS_FriendsList;
+	}
+	else if (FCString::Stricmp(InString, TEXT("Presence")) == 0)
+	{
+		OutEnum = EOS_EAuthScopeFlags::EOS_AS_Presence;
+	}
+	else if (FCString::Stricmp(InString, TEXT("FriendsManagement")) == 0)
+	{
+		OutEnum = EOS_EAuthScopeFlags::EOS_AS_FriendsManagement;
+	}
+	else if (FCString::Stricmp(InString, TEXT("Email")) == 0)
+	{
+		OutEnum = EOS_EAuthScopeFlags::EOS_AS_Email;
+	}
+	else if (FCString::Stricmp(InString, TEXT("NoFlags")) == 0 || FCString::Stricmp(InString, TEXT("None")) == 0)
+	{
+		OutEnum = EOS_EAuthScopeFlags::EOS_AS_NoFlags;
+	}
+	else
+	{
+		return false;
+	}
+	return true;
+}
+
+bool LexFromString(EOS_ELoginCredentialType& OutEnum, const TCHAR* const InString)
+{
+	if (FCString::Stricmp(InString, TEXT("ExchangeCode")) == 0)
+	{
+		OutEnum = EOS_ELoginCredentialType::EOS_LCT_ExchangeCode;
+	}
+	else if (FCString::Stricmp(InString, TEXT("PersistentAuth")) == 0)
+	{
+		OutEnum = EOS_ELoginCredentialType::EOS_LCT_PersistentAuth;
+	}
+	//else if (FCString::Stricmp(InString, TEXT("DeviceCode")) == 0) // DeviceCode is deprecated
+	//{
+	//	OutEnum = EOS_ELoginCredentialType::EOS_LCT_DeviceCode;
+	//}
+	else if (FCString::Stricmp(InString, TEXT("Developer")) == 0)
+	{
+		OutEnum = EOS_ELoginCredentialType::EOS_LCT_Developer;
+	}
+	else if (FCString::Stricmp(InString, TEXT("RefreshToken")) == 0)
+	{
+		OutEnum = EOS_ELoginCredentialType::EOS_LCT_RefreshToken;
+	}
+	else if (FCString::Stricmp(InString, TEXT("AccountPortal")) == 0)
+	{
+		OutEnum = EOS_ELoginCredentialType::EOS_LCT_AccountPortal;
+	}
+	else if (FCString::Stricmp(InString, TEXT("ExternalAuth")) == 0)
+	{
+		OutEnum = EOS_ELoginCredentialType::EOS_LCT_ExternalAuth;
+	}
+	else
+	{
+		return false;
+	}
+	return true;
+}
+
+TOnlineAsyncOpHandle<FAuthLogin> FAuthEOS::Login(FAuthLogin::Params&& Params)
+{
+	EOS_HPlatform PlatformHandle = SetupEOSSdk();
+	EOS_HAuth AuthHandle = EOS_Platform_GetAuthInterface(PlatformHandle);
+	check(AuthHandle != nullptr);
+
+	TOnlineAsyncOp<FAuthLogin>& Op = Services.OpCache.GetOp<FAuthLogin>(MoveTemp(Params));
+
+	EOS_Auth_LoginOptions LoginOptions = { };
+	LoginOptions.ApiVersion = EOS_AUTH_LOGIN_API_LATEST;
+	bool bContainsFlagsNone = false;
+	for (const FString& Scope : Op.GetParams().Scopes)
+	{
+		EOS_EAuthScopeFlags ScopeFlag;
+		if (LexFromString(ScopeFlag, *Scope))
+		{
+			if (ScopeFlag == EOS_EAuthScopeFlags::EOS_AS_NoFlags)
+			{
+				bContainsFlagsNone = true;
+			}
+			LoginOptions.ScopeFlags |= ScopeFlag;
+		}
+		else
+		{
+			UE_LOG(LogTemp, Warning, TEXT("Invalid ScopeFlag=[%s]"), *Scope);
+			Op.SetError(Errors::UnknownError());
+			return Op.GetHandle();
+		}
+	}
+	// TODO:  Where to put default scopes?
+	if (!bContainsFlagsNone && LoginOptions.ScopeFlags == EOS_EAuthScopeFlags::EOS_AS_NoFlags)
+	{
+		LoginOptions.ScopeFlags = EOS_EAuthScopeFlags::EOS_AS_BasicProfile | EOS_EAuthScopeFlags::EOS_AS_FriendsList | EOS_EAuthScopeFlags::EOS_AS_Presence;
+	}	
+
+	FEosAuthCredentials Credentials;
+	if (LexFromString(Credentials.Type, *Op.GetParams().CredentialsType))
+	{
+		switch (Credentials.Type)
+		{
+		case EOS_ELoginCredentialType::EOS_LCT_ExchangeCode:
+			// This is how the Epic launcher will pass credentials to you
+			Credentials.IdAnsi[0] = '\0';
+			FCStringAnsi::Strncpy(Credentials.TokenAnsi, TCHAR_TO_UTF8(*Op.GetParams().CredentialsToken), EOS_MAX_TOKEN_SIZE);
+			break;
+		case EOS_ELoginCredentialType::EOS_LCT_Developer:
+			// This is auth via the EOS auth tool
+			FCStringAnsi::Strncpy(Credentials.IdAnsi, TCHAR_TO_UTF8(*Op.GetParams().CredentialsId), EOS_OSS_STRING_BUFFER_LENGTH);
+			FCStringAnsi::Strncpy(Credentials.TokenAnsi, TCHAR_TO_UTF8(*Op.GetParams().CredentialsToken), EOS_MAX_TOKEN_SIZE);
+			break;
+		case EOS_ELoginCredentialType::EOS_LCT_AccountPortal:
+			// This is auth via the EOS Account Portal
+			Credentials.IdAnsi[0] = '\0';
+			Credentials.TokenAnsi[0] = '\0';
+			break;
+		case EOS_ELoginCredentialType::EOS_LCT_PersistentAuth:
+			// This is auth via stored credentials in EOS
+			Credentials.Id = nullptr;
+			Credentials.Token = nullptr;
+			break;
+		default:
+			UE_LOG(LogTemp, Warning, TEXT("Unsupported CredentialsType=[%s]"), *Op.GetParams().CredentialsType);
+			Op.SetError(Errors::UnknownError()); // TODO
+			return Op.GetHandle();
+		}
+	}
+	else
+	{
+		UE_LOG(LogTemp, Warning, TEXT("Invalid CredentialsType=[%s]"), *Op.GetParams().CredentialsType);
+		Op.SetError(Errors::UnknownError()); // TODO
+		return Op.GetHandle();
+	}
+
+	Op.Then([this, AuthHandle, LoginOptions, Credentials](TOnlineAsyncOp<FAuthLogin>& InAsyncOp) mutable
+		{
+			LoginOptions.Credentials = &Credentials;
+			return EOS_Async<EOS_Auth_LoginCallbackInfo>(InAsyncOp, EOS_Auth_Login, AuthHandle, LoginOptions);
+		})
+		.Then([this](TOnlineAsyncOp<FAuthLogin>& InAsyncOp, const EOS_Auth_LoginCallbackInfo* Data)
+		{
+			UE_LOG(LogTemp, Warning, TEXT("LoginResult: [%s]"), ANSI_TO_TCHAR(EOS_EResult_ToString(Data->ResultCode)));
+
+			if (Data->ResultCode == EOS_EResult::EOS_Success)
+			{
+				// Success
+				TSharedRef<FAccountInfoEOS> AccountInfo = MakeShared<FAccountInfoEOS>();
+				AccountInfo->LocalUserNum = InAsyncOp.GetParams().LocalUserNum;
+				AccountInfo->UserId = MakeEOSAccountId(Data->LocalUserId);
+
+				check(!AccountInfos.Contains(AccountInfo->UserId));
+				AccountInfos.Emplace(AccountInfo->UserId, AccountInfo);
+
+				FAuthLogin::Result Result = { AccountInfo };
+				InAsyncOp.SetResult(MoveTemp(Result));
+			}
+			else if (Data->ResultCode == EOS_EResult::EOS_InvalidUser && Data->ContinuanceToken != nullptr)
+			{
+				// 
+			}
+			else
+			{
+				FOnlineError Error;
+				if (Data->ResultCode == EOS_EResult::EOS_InvalidAuth)
+				{
+					Error = Errors::InvalidCredentials();
+				}
+				else
+				{
+					Error = Errors::Unknown();
+				}
+
+				InAsyncOp.SetError(MoveTemp(Error));
+			}
+		})
+		.Enqueue();
+	return Op.GetHandle();
+}
+
+TOnlineAsyncOpHandle<FAuthLogout> FAuthEOS::Logout(FAuthLogout::Params&& Params)
+{
+	EOS_HPlatform PlatformHandle = SetupEOSSdk();
+	EOS_HAuth AuthHandle = EOS_Platform_GetAuthInterface(PlatformHandle);
+	check(AuthHandle != nullptr);
+	
+	FAccountId ParamLocalUserId = Params.LocalUserId;
+	TOptional<EOS_EpicAccountId> AccountId = EOSAccountIdFromOnlineServiceAccountId(ParamLocalUserId);
+	TOnlineAsyncOp<FAuthLogout>& Op = Services.OpCache.GetOp<FAuthLogout>(MoveTemp(Params));
+	if (AccountId.IsSet() && AccountInfos.Contains(ParamLocalUserId))
+	{
+		// Should we destroy persistent auth first?
+		TOnlineChainableAsyncOp<FAuthLogout, void> NextOp = Op.Then([](TOnlineAsyncOp<FAuthLogout>& InAsyncOp){ return; });
+		if (Params.bDestroyAuth)
+		{
+			EOS_Auth_DeletePersistentAuthOptions DeletePersistentAuthOptions = {0};
+			DeletePersistentAuthOptions.ApiVersion = EOS_AUTH_DELETEPERSISTENTAUTH_API_LATEST;
+			DeletePersistentAuthOptions.RefreshToken = nullptr; // Is this needed?  Docs say it's needed for consoles
+			NextOp = NextOp.Then([this, AuthHandle, DeletePersistentAuthOptions](TOnlineAsyncOp<FAuthLogout>& InAsyncOp)
+				{
+					return EOS_Async<EOS_Auth_DeletePersistentAuthCallbackInfo>(InAsyncOp, EOS_Auth_DeletePersistentAuth, AuthHandle, DeletePersistentAuthOptions);
+				})
+				.Then([](TOnlineAsyncOp<FAuthLogout>& InAsyncOp, const EOS_Auth_DeletePersistentAuthCallbackInfo* Data)
+				{
+					UE_LOG(LogTemp, Warning, TEXT("DeletePersistentAuthResult: [%s]"), ANSI_TO_TCHAR(EOS_EResult_ToString(Data->ResultCode)));
+					// Regardless of success/failure, continue
+				});
+		}
+		// Logout
+		NextOp.Then([this, AuthHandle, AccountId](TOnlineAsyncOp<FAuthLogout>& InAsyncOp)
+			{
+				EOS_Auth_LogoutOptions LogoutOptions = { };
+				LogoutOptions.ApiVersion = EOS_AUTH_LOGOUT_API_LATEST;
+				LogoutOptions.LocalUserId = AccountId.GetValue();
+				return EOS_Async<EOS_Auth_LogoutCallbackInfo>(InAsyncOp, EOS_Auth_Logout, AuthHandle, LogoutOptions);
+			})
+			.Then([](TOnlineAsyncOp<FAuthLogout>& InAsyncOp, const EOS_Auth_LogoutCallbackInfo* Data)
+			{
+				UE_LOG(LogTemp, Warning, TEXT("LogoutResult: [%s]"), ANSI_TO_TCHAR(EOS_EResult_ToString(Data->ResultCode)));
+
+				if (Data->ResultCode == EOS_EResult::EOS_Success)
+				{
+					// Success
+					InAsyncOp.SetResult(FAuthLogout::Result());
+				}
+				else
+				{
+					// TODO: Error codes
+					FOnlineError Error;
+					InAsyncOp.SetError(MoveTemp(Error));
+				}
+			}).Enqueue();
+	}
+	else
+	{
+		// TODO: Error codes
+		Op.SetError(Errors::UnknownError());
+	}
+	
+	return Op.GetHandle();
+}
+
+TOnlineAsyncOpHandle<FAuthGenerateAuth> FAuthEOS::GenerateAuth(FAuthGenerateAuth::Params&& Params)
+{
+	TSharedRef<TOnlineAsyncOp<FAuthGenerateAuth>> AsyncOperation = MakeShared<TOnlineAsyncOp<FAuthGenerateAuth>>(Services, MoveTemp(Params));
+	return AsyncOperation->GetHandle();
+}
+
+TOnlineResult<FAuthGetAccountByLocalUserNum::Result> FAuthEOS::GetAccountByLocalUserNum(FAuthGetAccountByLocalUserNum::Params&& Params)
+{
+	TOnlineResult<FAccountId> LocalUserIdResult = GetAccountIdByLocalUserNum(Params.LocalUserNum);
+	if (LocalUserIdResult.IsOk())
+	{
+		FAuthGetAccountByLocalUserNum::Result Result;
+		Result.AccountInfo = AccountInfos.FindChecked(LocalUserIdResult.GetOkValue());
+		return TOnlineResult<FAuthGetAccountByLocalUserNum::Result>(Result);
+	}
+	else
+	{
+		return TOnlineResult<FAuthGetAccountByLocalUserNum::Result>(LocalUserIdResult.GetErrorValue());
+	}
+}
+
+TOnlineResult<FAuthGetAccountByAccountId::Result> FAuthEOS::GetAccountByAccountId(FAuthGetAccountByAccountId::Params&& Params)
+{
+	if (TSharedRef<FAccountInfoEOS>* const FoundAccount = AccountInfos.Find(Params.LocalUserId))
+	{
+		FAuthGetAccountByAccountId::Result Result;
+		Result.AccountInfo = *FoundAccount;
+		return TOnlineResult<FAuthGetAccountByAccountId::Result>(Result);
+	}
+	else
+	{
+		// TODO: proper error
+		return TOnlineResult<FAuthGetAccountByAccountId::Result>(Errors::UnknownError());
+	}
+}
+
+TOnlineResult<FAccountId> FAuthEOS::GetAccountIdByLocalUserNum(int32 LocalUserNum) const
+{
+	for (const TPair<FAccountId, TSharedRef<FAccountInfoEOS>>& AccountPair : AccountInfos)
+	{
+		if (AccountPair.Value->LocalUserNum == LocalUserNum)
+		{
+			TOnlineResult<FAccountId> Result(AccountPair.Key);
+			return Result;
+		}
+	}
+	TOnlineResult<FAccountId> Result(Errors::UnknownError()); // TODO: error code
+	return Result;
+}
+
+/* UE::Online */ }
+
+#endif // WITH_EOS_SDK
diff --git a/Engine/Plugins/Online/Experimental/OnlineServicesEOS/Source/Private/OnlineServicesEOS.cpp b/Engine/Plugins/Online/Experimental/OnlineServicesEOS/Source/Private/OnlineServicesEOS.cpp
new file mode 100644
--- /dev/null
+++ b/Engine/Plugins/Online/Experimental/OnlineServicesEOS/Source/Private/OnlineServicesEOS.cpp
@@ -0,0 +1,19 @@
+// Copyright Epic Games, Inc. All Rights Reserved.
+
+#include "OnlineServicesEOS.h"
+#include "AuthEOS.h"
+
+namespace UE::Online {
+
+FOnlineServicesEOS::FOnlineServicesEOS()
+{
+}
+
+void FOnlineServicesEOS::RegisterComponents()
+{
+#if WITH_EOS_SDK
+	Components.Register<FAuthEOS>(*this);
+#endif
+}
+
+/* UE::Online */ }
\ No newline at end of file
diff --git a/Engine/Plugins/Online/Experimental/OnlineServicesEOS/Source/Private/OnlineServicesEOSModule.cpp b/Engine/Plugins/Online/Experimental/OnlineServicesEOS/Source/Private/OnlineServicesEOSModule.cpp
new file mode 100644
--- /dev/null
+++ b/Engine/Plugins/Online/Experimental/OnlineServicesEOS/Source/Private/OnlineServicesEOSModule.cpp
@@ -0,0 +1,38 @@
+// Copyright Epic Games, Inc. All Rights Reserved.
+
+#include "CoreMinimal.h"
+#include "Modules/ModuleInterface.h"
+#include "Modules/ModuleManager.h"
+
+#include "Online/OnlineServicesRegistry.h"
+#include "OnlineServicesEOS.h"
+
+namespace UE::Online
+{
+
+class FOnlineServicesEOSModule : public IModuleInterface
+{
+public:
+	virtual void StartupModule() override;
+protected:
+};
+
+class FOnlineServicesFactoryEOS : public IOnlineServicesFactory
+{
+public:
+	virtual ~FOnlineServicesFactoryEOS() {}
+	virtual TSharedPtr<IOnlineServices> Create() override
+	{
+		return MakeShared<FOnlineServicesEOS>();
+	}
+protected:
+};
+
+void FOnlineServicesEOSModule::StartupModule()
+{
+	FOnlineServicesRegistry::Get().RegisterServicesFactory(EOnlineServices::Epic, MakeUnique<FOnlineServicesFactoryEOS>());
+}
+
+/* UE::Online */ }
+
+IMPLEMENT_MODULE(UE::Online::FOnlineServicesEOSModule, OnlineServicesEOS);
diff --git a/Engine/Plugins/Online/Experimental/OnlineServicesEOS/Source/Public/AuthEOS.h b/Engine/Plugins/Online/Experimental/OnlineServicesEOS/Source/Public/AuthEOS.h
new file mode 100644
--- /dev/null
+++ b/Engine/Plugins/Online/Experimental/OnlineServicesEOS/Source/Public/AuthEOS.h
@@ -0,0 +1,41 @@
+// Copyright Epic Games, Inc. All Rights Reserved.
+
+#pragma once
+
+#include "CoreMinimal.h"
+#include "Online/AuthCommon.h"
+#include "Containers/Ticker.h"
+
+#if WITH_EOS_SDK
+	#include "eos_auth_types.h"
+#endif
+namespace UE::Online {
+
+class FOnlineServicesEOS;
+
+class ONLINESERVICESEOS_API FAuthEOS : public FAuthCommon
+{
+public:
+	using Super = FAuthCommon;
+
+	FAuthEOS(FOnlineServicesEOS& InOwningSubsystem);
+	virtual TOnlineAsyncOpHandle<FAuthLogin> Login(FAuthLogin::Params&& Params) override;
+	virtual TOnlineAsyncOpHandle<FAuthLogout> Logout(FAuthLogout::Params&& Params) override;
+	virtual TOnlineAsyncOpHandle<FAuthGenerateAuth> GenerateAuth(FAuthGenerateAuth::Params&& Params) override;
+	virtual TOnlineResult<FAuthGetAccountByLocalUserNum::Result> GetAccountByLocalUserNum(FAuthGetAccountByLocalUserNum::Params&& Params) override;
+	virtual TOnlineResult<FAuthGetAccountByAccountId::Result> GetAccountByAccountId(FAuthGetAccountByAccountId::Params&& Params) override;
+
+protected:
+	TOnlineResult<FAccountId> GetAccountIdByLocalUserNum(int32 LocalUserNum) const;
+	virtual void Tick(float DeltaTime) override;
+
+	class FAccountInfoEOS : public FAccountInfo
+	{
+	public:
+
+	};
+
+	TMap<FAccountId, TSharedRef<FAccountInfoEOS>> AccountInfos;
+};
+
+/* UE::Online */ }
\ No newline at end of file
diff --git a/Engine/Plugins/Online/Experimental/OnlineServicesEOS/Source/Public/OnlineServicesEOS.h b/Engine/Plugins/Online/Experimental/OnlineServicesEOS/Source/Public/OnlineServicesEOS.h
new file mode 100644
--- /dev/null
+++ b/Engine/Plugins/Online/Experimental/OnlineServicesEOS/Source/Public/OnlineServicesEOS.h
@@ -0,0 +1,19 @@
+// Copyright Epic Games, Inc. All Rights Reserved.
+
+#pragma once
+
+#include "Online/OnlineServicesCommon.h"
+
+namespace UE::Online {
+
+class ONLINESERVICESEOS_API FOnlineServicesEOS : public FOnlineServicesCommon
+{
+public:
+	FOnlineServicesEOS();
+	virtual void RegisterComponents() override;
+
+	static EOnlineServices GetServicesProvider() { return EOnlineServices::Epic; }
+private:
+};
+
+/* UE::Online */ }
diff --git a/Engine/Plugins/Online/Experimental/OnlineServicesEOS/Source/Public/OnlineServicesEOSTypes.h b/Engine/Plugins/Online/Experimental/OnlineServicesEOS/Source/Public/OnlineServicesEOSTypes.h
new file mode 100644
--- /dev/null
+++ b/Engine/Plugins/Online/Experimental/OnlineServicesEOS/Source/Public/OnlineServicesEOSTypes.h
@@ -0,0 +1,94 @@
+// Copyright Epic Games, Inc. All Rights Reserved.
+
+#pragma once
+
+#include "Online/OnlineAsyncOp.h"
+#include "Async/Future.h"
+
+#if WITH_EOS_SDK
+
+#include "eos_common.h"
+
+namespace UE::Online {
+
+/** Class to handle all callbacks generically using a future to forward callback results */
+template<typename CallbackType>
+class TEOSCallback
+{
+	using CallbackFuncType = void (EOS_CALL*)(const CallbackType*);
+public:
+	TEOSCallback() = default;
+	virtual ~TEOSCallback() = default;
+
+	operator CallbackFuncType()
+	{
+		return &CallbackImpl;
+	}
+
+	TFuture<const CallbackType*> GetFuture() { return Promise.GetFuture(); }
+
+private:
+	TPromise<const CallbackType*> Promise;
+
+	static void EOS_CALL CallbackImpl(const CallbackType* Data)
+	{
+		if (EOS_EResult_IsOperationComplete(Data->ResultCode) == EOS_FALSE)
+		{
+			// Ignore
+			return;
+		}
+		check(IsInGameThread());
+
+		TEOSCallback* CallbackThis = reinterpret_cast<TEOSCallback*>(Data->ClientData);
+		check(CallbackThis);
+
+		CallbackThis->Promise.EmplaceValue(Data);
+
+		delete CallbackThis;
+	}
+};
+
+template<typename TEosResult, typename TEosHandle, typename TEosParameters, typename TEosFn, typename TAsyncOpType> 
+TFuture<const TEosResult*> EOS_Async(TOnlineAsyncOp<TAsyncOpType>& Op, TEosFn EosFn, TEosHandle EosHandle, TEosParameters Parameters)
+{
+	TEOSCallback<TEosResult>* Callback = new TEOSCallback<TEosResult>();
+	EosFn(EosHandle, &Parameters, Callback, *Callback);
+	return Callback->GetFuture();
+}
+
+// TEMP until Net Id Registry is done
+extern TMap<EOS_EpicAccountId, int32> EOSAccountIdMap;
+FAccountId MakeEOSAccountId(EOS_EpicAccountId EpicAccountId)
+{
+	static int32 EpicAccountIdCounter = 0;
+
+	FAccountId Result;
+	//Result.Type = EOnlineServices::Epic;
+	if (int32* ExistingId = EOSAccountIdMap.Find(EpicAccountId))
+	{
+		Result.Handle = *ExistingId;
+	}
+	else
+	{
+		Result.Handle = EOSAccountIdMap.Emplace(EpicAccountId, ++EpicAccountIdCounter);
+	}
+	return Result;
+}
+
+TOptional<EOS_EpicAccountId> EOSAccountIdFromOnlineServiceAccountId(const FAccountId& InAccountId)
+{
+	TOptional<EOS_EpicAccountId> Result;
+	for (const TPair<EOS_EpicAccountId, int32>& EOSAccountIdPair : EOSAccountIdMap)
+	{
+		if (EOSAccountIdPair.Value == InAccountId.Handle)
+		{
+			Result = EOSAccountIdPair.Key;
+			break;
+		}
+	}
+	return Result;
+}
+
+/* UE::Online */ }
+
+#endif // WITH_EOS_SDK
