ClangEditor fix
Wrap entire AuthEOS interface with WITH_EOS_SDK since the implementation is wrapped with that
horde 83612
Checked in on behalf of Rob Cannaday

diff --git a/Engine/Plugins/Online/Experimental/OnlineServicesEOS/Source/Public/AuthEOS.h b/Engine/Plugins/Online/Experimental/OnlineServicesEOS/Source/Public/AuthEOS.h
--- a/Engine/Plugins/Online/Experimental/OnlineServicesEOS/Source/Public/AuthEOS.h
+++ b/Engine/Plugins/Online/Experimental/OnlineServicesEOS/Source/Public/AuthEOS.h
@@ -7,8 +7,8 @@
 #include "Containers/Ticker.h"
 
 #if WITH_EOS_SDK
-	#include "eos_auth_types.h"
-#endif
+#include "eos_auth_types.h"
+
 namespace UE::Online {
 
 class FOnlineServicesEOS;
@@ -38,4 +38,5 @@ protected:
 	TMap<FAccountId, TSharedRef<FAccountInfoEOS>> AccountInfos;
 };
 
-/* UE::Online */ }
\ No newline at end of file
+/* UE::Online */ }
+#endif
