Modify OnlineServicesEOS to use EOSShared

diff --git a/Engine/Plugins/Online/Experimental/OnlineServicesEOS/OnlineServicesEOS.uplugin b/Engine/Plugins/Online/Experimental/OnlineServicesEOS/OnlineServicesEOS.uplugin
--- a/Engine/Plugins/Online/Experimental/OnlineServicesEOS/OnlineServicesEOS.uplugin
+++ b/Engine/Plugins/Online/Experimental/OnlineServicesEOS/OnlineServicesEOS.uplugin
@@ -27,6 +27,10 @@
 		{
 			"Name": "OnlineServices",
 			"Enabled": true
+		},
+		{
+			"Name": "EOSShared",
+			"Enabled": true
 		}
 	]
 }
\ No newline at end of file
diff --git a/Engine/Plugins/Online/Experimental/OnlineServicesEOS/Source/OnlineServicesEOS.Build.cs b/Engine/Plugins/Online/Experimental/OnlineServicesEOS/Source/OnlineServicesEOS.Build.cs
--- a/Engine/Plugins/Online/Experimental/OnlineServicesEOS/Source/OnlineServicesEOS.Build.cs
+++ b/Engine/Plugins/Online/Experimental/OnlineServicesEOS/Source/OnlineServicesEOS.Build.cs
@@ -15,7 +15,11 @@ public class OnlineServicesEOS : ModuleRules
 			}
 		);
 
-		// TODO:  Use EOSShared module
-		PrivateDefinitions.Add("WITH_EOS_SDK=0");
+		PrivateDependencyModuleNames.AddRange(
+			new string[] {
+				"EOSSDK",
+				"EOSShared",
+			}
+		);
 	}
 }
diff --git a/Engine/Plugins/Online/Experimental/OnlineServicesEOS/Source/Private/AuthEOS.cpp b/Engine/Plugins/Online/Experimental/OnlineServicesEOS/Source/Private/AuthEOS.cpp
--- a/Engine/Plugins/Online/Experimental/OnlineServicesEOS/Source/Private/AuthEOS.cpp
+++ b/Engine/Plugins/Online/Experimental/OnlineServicesEOS/Source/Private/AuthEOS.cpp
@@ -1,118 +1,18 @@
 // Copyright Epic Games, Inc. All Rights Reserved.
 
 #include "AuthEOS.h"
+
+#if WITH_EOS_SDK
 #include "OnlineServicesEOS.h"
 #include "OnlineServicesEOSTypes.h"
 #include "Online/AuthErrors.h"
 
-#if WITH_EOS_SDK
 #include "eos_auth.h"
 #include "eos_types.h"
 #include "eos_init.h"
 #include "eos_sdk.h"
 #include "eos_logging.h"
 
-void EOS_CALL EOSFree(void* Ptr)
-{
-	FMemory::Free(Ptr);
-}
-
-void* EOS_CALL EOSMalloc(size_t Size, size_t Alignment)
-{
-	return FMemory::Malloc(Size, Alignment);
-}
-
-void* EOS_CALL EOSRealloc(void* Ptr, size_t Size, size_t Alignment)
-{
-	return FMemory::Realloc(Ptr, Size, Alignment);
-}
-
-void EOS_CALL EOSLog(const EOS_LogMessage* InMsg)
-{
-	if (GLog == nullptr)
-	{
-		return;
-	}
-
-	switch (InMsg->Level)
-	{
-		case EOS_ELogLevel::EOS_LOG_Fatal:
-		{
-			UE_LOG(LogTemp, Warning, TEXT("EOSSDK-%s: %s"), ANSI_TO_TCHAR(InMsg->Category), ANSI_TO_TCHAR(InMsg->Message));
-			break;
-		}
-		case EOS_ELogLevel::EOS_LOG_Error:
-		{
-			UE_LOG(LogTemp, Warning, TEXT("EOSSDK-%s: %s"), ANSI_TO_TCHAR(InMsg->Category), ANSI_TO_TCHAR(InMsg->Message));
-			break;
-		}
-		case EOS_ELogLevel::EOS_LOG_Warning:
-		{
-			UE_LOG(LogTemp, Warning, TEXT("EOSSDK-%s: %s"), ANSI_TO_TCHAR(InMsg->Category), ANSI_TO_TCHAR(InMsg->Message));
-			break;
-		}
-		case EOS_ELogLevel::EOS_LOG_Verbose:
-		{
-			UE_LOG(LogTemp, Warning, TEXT("EOSSDK-%s: %s"), ANSI_TO_TCHAR(InMsg->Category), ANSI_TO_TCHAR(InMsg->Message));
-			break;
-		}
-		case EOS_ELogLevel::EOS_LOG_VeryVerbose:
-		{
-			UE_LOG(LogTemp, Warning, TEXT("EOSSDK-%s: %s"), ANSI_TO_TCHAR(InMsg->Category), ANSI_TO_TCHAR(InMsg->Message));
-			break;
-		}
-		case EOS_ELogLevel::EOS_LOG_Info:
-		default:
-		{
-			UE_LOG(LogTemp, Warning, TEXT("EOSSDK-%s: %s"), ANSI_TO_TCHAR(InMsg->Category), ANSI_TO_TCHAR(InMsg->Message));
-			break;
-		}
-	}
-}
-
-EOS_HPlatform SetupEOSSdk()
-{
-	static EOS_HPlatform PlatformHandle = nullptr;
-	if (PlatformHandle == nullptr)
-	{
-		// Init EOS SDK
-		EOS_InitializeOptions SDKOptions = { };
-		SDKOptions.ApiVersion = EOS_INITIALIZE_API_LATEST;
-
-		SDKOptions.ProductName = "My Product";
-		SDKOptions.ProductVersion = "Unknown";
-		SDKOptions.AllocateMemoryFunction = &EOSMalloc;
-		SDKOptions.ReallocateMemoryFunction = &EOSRealloc;
-		SDKOptions.ReleaseMemoryFunction = &EOSFree;
-
-		EOS_EResult InitResult = EOS_Initialize(&SDKOptions);
-
-		EOS_Platform_Options PlatformOptions = {};
-		PlatformOptions.ApiVersion = EOS_PLATFORM_OPTIONS_API_LATEST;
-		PlatformOptions.bIsServer = EOS_FALSE;
-		PlatformOptions.EncryptionKey = nullptr;
-		PlatformOptions.OverrideCountryCode = nullptr;
-		PlatformOptions.OverrideLocaleCode = nullptr;
-		PlatformOptions.Flags = EOS_PF_WINDOWS_ENABLE_OVERLAY_D3D9 | EOS_PF_WINDOWS_ENABLE_OVERLAY_D3D10 | EOS_PF_WINDOWS_ENABLE_OVERLAY_OPENGL; // Enable overlay support for D3D9/10 and OpenGL. This sample uses D3D11 or SDL.
-
-		char CacheDirectory[512];
-		FCStringAnsi::Strncpy(CacheDirectory, TCHAR_TO_UTF8(*(FPlatformProcess::UserDir() / FString(TEXT("CacheDir")))), 512);
-		PlatformOptions.CacheDirectory = CacheDirectory;
-
-		PlatformOptions.ProductId = "";
-		PlatformOptions.SandboxId = "";
-		PlatformOptions.DeploymentId = "";
-
-		PlatformOptions.ClientCredentials.ClientId = "";
-		PlatformOptions.ClientCredentials.ClientSecret = "";
-
-		PlatformHandle = EOS_Platform_Create(&PlatformOptions);
-
-		EOS_Logging_SetCallback(&EOSLog);
-	}
-	return PlatformHandle;
-}
-
 namespace UE::Online {
 
 // TEMP until Net Id Registry is done
@@ -152,18 +52,11 @@ struct FEosAuthCredentials :
 	char TokenAnsi[EOS_MAX_TOKEN_SIZE];
 };
 
-FAuthEOS::FAuthEOS(FOnlineServicesEOS& InOwningSubsystem)
-	: FAuthCommon(InOwningSubsystem)
-{
-}
-
-void FAuthEOS::Tick(float DeltaTime)
+FAuthEOS::FAuthEOS(FOnlineServicesEOS& InServices)
+	: FAuthCommon(InServices)
 {
-	EOS_HPlatform PlatformHandle = SetupEOSSdk();
-	if (PlatformHandle)
-	{
-		EOS_Platform_Tick(PlatformHandle);
-	}
+	AuthHandle = EOS_Platform_GetAuthInterface(InServices.GetEOSPlatformHandle());
+	check(AuthHandle != nullptr);
 }
 
 bool LexFromString(EOS_EAuthScopeFlags& OutEnum, const TCHAR* InString)
@@ -238,10 +131,6 @@ bool LexFromString(EOS_ELoginCredentialType& OutEnum, const TCHAR* const InStrin
 
 TOnlineAsyncOpHandle<FAuthLogin> FAuthEOS::Login(FAuthLogin::Params&& Params)
 {
-	EOS_HPlatform PlatformHandle = SetupEOSSdk();
-	EOS_HAuth AuthHandle = EOS_Platform_GetAuthInterface(PlatformHandle);
-	check(AuthHandle != nullptr);
-
 	TOnlineAsyncOp<FAuthLogin>& Op = Services.OpCache.GetOp<FAuthLogin>(MoveTemp(Params));
 
 	EOS_Auth_LoginOptions LoginOptions = { };
@@ -309,7 +198,7 @@ TOnlineAsyncOpHandle<FAuthLogin> FAuthEOS::Login(FAuthLogin::Params&& Params)
 		return Op.GetHandle();
 	}
 
-	Op.Then([this, AuthHandle, LoginOptions, Credentials](TOnlineAsyncOp<FAuthLogin>& InAsyncOp) mutable
+	Op.Then([this, LoginOptions, Credentials](TOnlineAsyncOp<FAuthLogin>& InAsyncOp) mutable
 		{
 			LoginOptions.Credentials = &Credentials;
 			return EOS_Async<EOS_Auth_LoginCallbackInfo>(InAsyncOp, EOS_Auth_Login, AuthHandle, LoginOptions);
@@ -356,10 +245,6 @@ TOnlineAsyncOpHandle<FAuthLogin> FAuthEOS::Login(FAuthLogin::Params&& Params)
 
 TOnlineAsyncOpHandle<FAuthLogout> FAuthEOS::Logout(FAuthLogout::Params&& Params)
 {
-	EOS_HPlatform PlatformHandle = SetupEOSSdk();
-	EOS_HAuth AuthHandle = EOS_Platform_GetAuthInterface(PlatformHandle);
-	check(AuthHandle != nullptr);
-	
 	FAccountId ParamLocalUserId = Params.LocalUserId;
 	TOptional<EOS_EpicAccountId> AccountId = EOSAccountIdFromOnlineServiceAccountId(ParamLocalUserId);
 	TOnlineAsyncOp<FAuthLogout>& Op = Services.OpCache.GetOp<FAuthLogout>(MoveTemp(Params));
@@ -372,7 +257,7 @@ TOnlineAsyncOpHandle<FAuthLogout> FAuthEOS::Logout(FAuthLogout::Params&& Params)
 			EOS_Auth_DeletePersistentAuthOptions DeletePersistentAuthOptions = {0};
 			DeletePersistentAuthOptions.ApiVersion = EOS_AUTH_DELETEPERSISTENTAUTH_API_LATEST;
 			DeletePersistentAuthOptions.RefreshToken = nullptr; // Is this needed?  Docs say it's needed for consoles
-			NextOp = NextOp.Then([this, AuthHandle, DeletePersistentAuthOptions](TOnlineAsyncOp<FAuthLogout>& InAsyncOp)
+			NextOp = NextOp.Then([this, DeletePersistentAuthOptions](TOnlineAsyncOp<FAuthLogout>& InAsyncOp)
 				{
 					return EOS_Async<EOS_Auth_DeletePersistentAuthCallbackInfo>(InAsyncOp, EOS_Auth_DeletePersistentAuth, AuthHandle, DeletePersistentAuthOptions);
 				})
@@ -383,7 +268,7 @@ TOnlineAsyncOpHandle<FAuthLogout> FAuthEOS::Logout(FAuthLogout::Params&& Params)
 				});
 		}
 		// Logout
-		NextOp.Then([this, AuthHandle, AccountId](TOnlineAsyncOp<FAuthLogout>& InAsyncOp)
+		NextOp.Then([this, AccountId](TOnlineAsyncOp<FAuthLogout>& InAsyncOp)
 			{
 				EOS_Auth_LogoutOptions LogoutOptions = { };
 				LogoutOptions.ApiVersion = EOS_AUTH_LOGOUT_API_LATEST;
diff --git a/Engine/Plugins/Online/Experimental/OnlineServicesEOS/Source/Private/OnlineServicesEOS.cpp b/Engine/Plugins/Online/Experimental/OnlineServicesEOS/Source/Private/OnlineServicesEOS.cpp
--- a/Engine/Plugins/Online/Experimental/OnlineServicesEOS/Source/Private/OnlineServicesEOS.cpp
+++ b/Engine/Plugins/Online/Experimental/OnlineServicesEOS/Source/Private/OnlineServicesEOS.cpp
@@ -1,12 +1,41 @@
 // Copyright Epic Games, Inc. All Rights Reserved.
 
 #include "OnlineServicesEOS.h"
+
+#if WITH_EOS_SDK
+
+#include "IEOSSDKManager.h"
 #include "AuthEOS.h"
 
 namespace UE::Online {
 
 FOnlineServicesEOS::FOnlineServicesEOS()
 {
+	// Init EOS SDK
+	EOS_Platform_Options PlatformOptions = {};
+	PlatformOptions.ApiVersion = EOS_PLATFORM_OPTIONS_API_LATEST;
+	PlatformOptions.bIsServer = EOS_FALSE;
+	PlatformOptions.EncryptionKey = nullptr;
+	PlatformOptions.OverrideCountryCode = nullptr;
+	PlatformOptions.OverrideLocaleCode = nullptr;
+	PlatformOptions.Flags = EOS_PF_WINDOWS_ENABLE_OVERLAY_D3D9 | EOS_PF_WINDOWS_ENABLE_OVERLAY_D3D10 | EOS_PF_WINDOWS_ENABLE_OVERLAY_OPENGL; // Enable overlay support for D3D9/10 and OpenGL. This sample uses D3D11 or SDL.
+
+	char CacheDirectory[512];
+	FCStringAnsi::Strncpy(CacheDirectory, TCHAR_TO_UTF8(*(FPlatformProcess::UserDir() / FString(TEXT("CacheDir")))), 512);
+	PlatformOptions.CacheDirectory = CacheDirectory;
+
+	PlatformOptions.ProductId = "";
+	PlatformOptions.SandboxId = "";
+	PlatformOptions.DeploymentId = "";
+
+	PlatformOptions.ClientCredentials.ClientId = "";
+	PlatformOptions.ClientCredentials.ClientSecret = "";
+
+	IEOSSDKManager* SDKManager = IEOSSDKManager::Get();
+	if (ensure(SDKManager))
+	{
+		EOSPlatformHandle = SDKManager->CreatePlatform(PlatformOptions);
+	}
 }
 
 void FOnlineServicesEOS::RegisterComponents()
@@ -16,4 +45,16 @@ void FOnlineServicesEOS::RegisterComponents()
 #endif
 }
 
-/* UE::Online */ }
\ No newline at end of file
+EOS_HPlatform FOnlineServicesEOS::GetEOSPlatformHandle() const
+{
+	if (EOSPlatformHandle)
+	{
+		return *EOSPlatformHandle;
+	}
+	return nullptr;
+}
+
+
+/* UE::Online */ }
+
+#endif // WITH_EOS_SDK
\ No newline at end of file
diff --git a/Engine/Plugins/Online/Experimental/OnlineServicesEOS/Source/Private/OnlineServicesEOSModule.cpp b/Engine/Plugins/Online/Experimental/OnlineServicesEOS/Source/Private/OnlineServicesEOSModule.cpp
--- a/Engine/Plugins/Online/Experimental/OnlineServicesEOS/Source/Private/OnlineServicesEOSModule.cpp
+++ b/Engine/Plugins/Online/Experimental/OnlineServicesEOS/Source/Private/OnlineServicesEOSModule.cpp
@@ -17,6 +17,7 @@ public:
 protected:
 };
 
+#if WITH_EOS_SDK
 class FOnlineServicesFactoryEOS : public IOnlineServicesFactory
 {
 public:
@@ -27,10 +28,13 @@ public:
 	}
 protected:
 };
+#endif // WITH_EOS_SDK
 
 void FOnlineServicesEOSModule::StartupModule()
 {
+#if WITH_EOS_SDK
 	FOnlineServicesRegistry::Get().RegisterServicesFactory(EOnlineServices::Epic, MakeUnique<FOnlineServicesFactoryEOS>());
+#endif // WITH_EOS_SDK
 }
 
 /* UE::Online */ }
diff --git a/Engine/Plugins/Online/Experimental/OnlineServicesEOS/Source/Public/AuthEOS.h b/Engine/Plugins/Online/Experimental/OnlineServicesEOS/Source/Public/AuthEOS.h
--- a/Engine/Plugins/Online/Experimental/OnlineServicesEOS/Source/Public/AuthEOS.h
+++ b/Engine/Plugins/Online/Experimental/OnlineServicesEOS/Source/Public/AuthEOS.h
@@ -2,11 +2,10 @@
 
 #pragma once
 
+#if WITH_EOS_SDK
 #include "CoreMinimal.h"
 #include "Online/AuthCommon.h"
-#include "Containers/Ticker.h"
 
-#if WITH_EOS_SDK
 #include "eos_auth_types.h"
 
 namespace UE::Online {
@@ -25,9 +24,10 @@ public:
 	virtual TOnlineResult<FAuthGetAccountByLocalUserNum::Result> GetAccountByLocalUserNum(FAuthGetAccountByLocalUserNum::Params&& Params) override;
 	virtual TOnlineResult<FAuthGetAccountByAccountId::Result> GetAccountByAccountId(FAuthGetAccountByAccountId::Params&& Params) override;
 
+	bool IsLoggedIn(const FAccountId& AccountId) const;
+
 protected:
 	TOnlineResult<FAccountId> GetAccountIdByLocalUserNum(int32 LocalUserNum) const;
-	virtual void Tick(float DeltaTime) override;
 
 	class FAccountInfoEOS : public FAccountInfo
 	{
@@ -36,7 +36,9 @@ protected:
 	};
 
 	TMap<FAccountId, TSharedRef<FAccountInfoEOS>> AccountInfos;
+
+	EOS_HAuth AuthHandle;
 };
 
 /* UE::Online */ }
-#endif
+#endif // WITH_EOS_SDK
\ No newline at end of file
diff --git a/Engine/Plugins/Online/Experimental/OnlineServicesEOS/Source/Public/OnlineServicesEOS.h b/Engine/Plugins/Online/Experimental/OnlineServicesEOS/Source/Public/OnlineServicesEOS.h
--- a/Engine/Plugins/Online/Experimental/OnlineServicesEOS/Source/Public/OnlineServicesEOS.h
+++ b/Engine/Plugins/Online/Experimental/OnlineServicesEOS/Source/Public/OnlineServicesEOS.h
@@ -4,6 +4,12 @@
 
 #include "Online/OnlineServicesCommon.h"
 
+#if WITH_EOS_SDK
+
+using IEOSPlatformHandlePtr = TSharedPtr<class IEOSPlatformHandle>;
+
+#include "eos_sdk.h"
+
 namespace UE::Online {
 
 class ONLINESERVICESEOS_API FOnlineServicesEOS : public FOnlineServicesCommon
@@ -13,7 +19,12 @@ public:
 	virtual void RegisterComponents() override;
 
 	static EOnlineServices GetServicesProvider() { return EOnlineServices::Epic; }
-private:
+
+	EOS_HPlatform GetEOSPlatformHandle() const;
+protected:
+	IEOSPlatformHandlePtr EOSPlatformHandle;
 };
 
 /* UE::Online */ }
+
+#endif // WITH_EOS_SDK
\ No newline at end of file
diff --git a/Engine/Plugins/Online/Experimental/OnlineServicesEOS/Source/Public/OnlineServicesEOSTypes.h b/Engine/Plugins/Online/Experimental/OnlineServicesEOS/Source/Public/OnlineServicesEOSTypes.h
--- a/Engine/Plugins/Online/Experimental/OnlineServicesEOS/Source/Public/OnlineServicesEOSTypes.h
+++ b/Engine/Plugins/Online/Experimental/OnlineServicesEOS/Source/Public/OnlineServicesEOSTypes.h
@@ -2,11 +2,11 @@
 
 #pragma once
 
+#if WITH_EOS_SDK
+
 #include "Online/OnlineAsyncOp.h"
 #include "Async/Future.h"
 
-#if WITH_EOS_SDK
-
 #include "eos_common.h"
 
 namespace UE::Online {
@@ -58,7 +58,7 @@ TFuture<const TEosResult*> EOS_Async(TOnlineAsyncOp<TAsyncOpType>& Op, TEosFn Eo
 
 // TEMP until Net Id Registry is done
 extern TMap<EOS_EpicAccountId, int32> EOSAccountIdMap;
-FAccountId MakeEOSAccountId(EOS_EpicAccountId EpicAccountId)
+inline FAccountId MakeEOSAccountId(EOS_EpicAccountId EpicAccountId)
 {
 	static int32 EpicAccountIdCounter = 0;
 
@@ -75,7 +75,7 @@ FAccountId MakeEOSAccountId(EOS_EpicAccountId EpicAccountId)
 	return Result;
 }
 
-TOptional<EOS_EpicAccountId> EOSAccountIdFromOnlineServiceAccountId(const FAccountId& InAccountId)
+inline TOptional<EOS_EpicAccountId> EOSAccountIdFromOnlineServiceAccountId(const FAccountId& InAccountId)
 {
 	TOptional<EOS_EpicAccountId> Result;
 	for (const TPair<EOS_EpicAccountId, int32>& EOSAccountIdPair : EOSAccountIdMap)
