UE-124571 //UE5/Main - Compile UnrealGame HoloLens - Source file 'EOSSDK-HoloLens-Shipping.dll' does not exist
-Blacklisting the Hololens platform in OnlineSerivesEOS.uplugin.

diff --git a/Engine/Plugins/Online/Experimental/OnlineServicesEOS/OnlineServicesEOS.uplugin b/Engine/Plugins/Online/Experimental/OnlineServicesEOS/OnlineServicesEOS.uplugin
--- a/Engine/Plugins/Online/Experimental/OnlineServicesEOS/OnlineServicesEOS.uplugin
+++ b/Engine/Plugins/Online/Experimental/OnlineServicesEOS/OnlineServicesEOS.uplugin
@@ -16,7 +16,8 @@
 			"Type": "Runtime",
 			"LoadingPhase": "PostConfigInit",
 			"BlacklistPlatforms": [
-				"LinuxArm64"
+				"LinuxArm64",
+				"Hololens"
 			]
 		}
 	],
