- Blacklist TVOS

#horde 96358

diff --git a/Engine/Plugins/Online/Experimental/OnlineServicesEOS/OnlineServicesEOS.uplugin b/Engine/Plugins/Online/Experimental/OnlineServicesEOS/OnlineServicesEOS.uplugin
--- a/Engine/Plugins/Online/Experimental/OnlineServicesEOS/OnlineServicesEOS.uplugin
+++ b/Engine/Plugins/Online/Experimental/OnlineServicesEOS/OnlineServicesEOS.uplugin
@@ -17,7 +17,8 @@
 			"LoadingPhase": "PostConfigInit",
 			"BlacklistPlatforms": [
 				"LinuxArm64",
-				"Hololens"
+				"Hololens",
+				"TVOS"
 			]
 		}
 	],
