- First pass of net id registry integration into OnlineServices.
- OSS EOS has its own bespoke AccountId registry type due to it having 2 underlying account id types

diff --git a/Engine/Plugins/Online/EOSShared/Source/EOSShared/Private/EOSShared.cpp b/Engine/Plugins/Online/EOSShared/Source/EOSShared/Private/EOSShared.cpp
--- a/Engine/Plugins/Online/EOSShared/Source/EOSShared/Private/EOSShared.cpp
+++ b/Engine/Plugins/Online/EOSShared/Source/EOSShared/Private/EOSShared.cpp
@@ -2,6 +2,8 @@
 
 #include "EOSShared.h"
 
+#include "eos_friends_types.h"
+
 DEFINE_LOG_CATEGORY(LogEOSSDK);
 FString LexToString(const EOS_ProductUserId UserId)
 {
@@ -34,3 +36,15 @@ FString LexToString(const EOS_EpicAccountId AccountId)
 
 	return Result;
 }
+
+const TCHAR* LexToString(const EOS_EFriendsStatus FriendStatus)
+{
+	switch (FriendStatus)
+	{
+		default: checkNoEntry(); // Intentional fall through
+		case EOS_EFriendsStatus::EOS_FS_NotFriends:		return TEXT("NotFriends");
+		case EOS_EFriendsStatus::EOS_FS_InviteSent:		return TEXT("InviteSent");
+		case EOS_EFriendsStatus::EOS_FS_InviteReceived: return TEXT("InviteReceived");
+		case EOS_EFriendsStatus::EOS_FS_Friends:		return TEXT("Friends");
+	}
+} 
\ No newline at end of file
diff --git a/Engine/Plugins/Online/EOSShared/Source/EOSShared/Public/EOSShared.h b/Engine/Plugins/Online/EOSShared/Source/EOSShared/Public/EOSShared.h
--- a/Engine/Plugins/Online/EOSShared/Source/EOSShared/Public/EOSShared.h
+++ b/Engine/Plugins/Online/EOSShared/Source/EOSShared/Public/EOSShared.h
@@ -10,6 +10,7 @@
 #if defined(EOS_PLATFORM_BASE_FILE_NAME)
 #include EOS_PLATFORM_BASE_FILE_NAME
 #endif
+#include "eos_base.h"
 #include "eos_common.h"
 #include "eos_version.h"
 
@@ -19,6 +20,10 @@
 #define WITH_EOS_RTC WITH_EOS_SDK && (EOS_MAJOR_VERSION >= 1 && EOS_MINOR_VERSION >= 13)
 #endif
 
+#define EOS_ENUM_FORWARD_DECL(name) enum class name : int32_t;
+EOS_ENUM_FORWARD_DECL(EOS_EFriendsStatus);
+#undef EOS_ENUM_FORWARD_DECL
+
 DECLARE_LOG_CATEGORY_EXTERN(LogEOSSDK, Log, All);
 
 inline FString LexToString(EOS_EResult EosResult)
@@ -26,7 +31,8 @@ inline FString LexToString(EOS_EResult EosResult)
 	return ANSI_TO_TCHAR(EOS_EResult_ToString(EosResult));
 }
 
-FString EOSSHARED_API LexToString(const EOS_ProductUserId UserId);
-FString EOSSHARED_API LexToString(const EOS_EpicAccountId AccountId);
+EOSSHARED_API FString LexToString(const EOS_ProductUserId UserId);
+EOSSHARED_API FString LexToString(const EOS_EpicAccountId AccountId);
+EOSSHARED_API const TCHAR* LexToString(const EOS_EFriendsStatus FriendStatus);
 
 #endif // WITH_EOS_SDK
\ No newline at end of file
diff --git a/Engine/Plugins/Online/Experimental/OnlineServices/Source/OnlineServicesCommon/Private/Online/OnlineServicesCommon.cpp b/Engine/Plugins/Online/Experimental/OnlineServices/Source/OnlineServicesCommon/Private/Online/OnlineServicesCommon.cpp
--- a/Engine/Plugins/Online/Experimental/OnlineServices/Source/OnlineServicesCommon/Private/Online/OnlineServicesCommon.cpp
+++ b/Engine/Plugins/Online/Experimental/OnlineServices/Source/OnlineServicesCommon/Private/Online/OnlineServicesCommon.cpp
@@ -39,11 +39,6 @@ void FOnlineServicesCommon::Destroy()
 	Shutdown();
 }
 
-FAccountId FOnlineServicesCommon::CreateAccountId(FString&& InAccountIdString)
-{
-	return FAccountId();
-}
-
 IAuthPtr FOnlineServicesCommon::GetAuthInterface()
 {
 	return IAuthPtr(AsShared(), Get<IAuth>());
@@ -112,7 +107,7 @@ FOnlineAsyncOpQueue& FOnlineServicesCommon::GetSerialQueue()
 	return SerialQueue;
 }
 
-FOnlineAsyncOpQueue& FOnlineServicesCommon::GetSerialQueue(FAccountId& AccountId)
+FOnlineAsyncOpQueue& FOnlineServicesCommon::GetSerialQueue(const FOnlineAccountIdHandle& AccountId)
 {
 	TUniquePtr<FOnlineAsyncOpQueueSerial>* Queue = PerUserSerialQueue.Find(AccountId);
 	if (Queue == nullptr)
diff --git a/Engine/Plugins/Online/Experimental/OnlineServices/Source/OnlineServicesCommon/Private/Online/OnlineSyntaxTests.cpp b/Engine/Plugins/Online/Experimental/OnlineServices/Source/OnlineServicesCommon/Private/Online/OnlineSyntaxTests.cpp
--- a/Engine/Plugins/Online/Experimental/OnlineServices/Source/OnlineServicesCommon/Private/Online/OnlineSyntaxTests.cpp
+++ b/Engine/Plugins/Online/Experimental/OnlineServices/Source/OnlineServicesCommon/Private/Online/OnlineSyntaxTests.cpp
@@ -47,7 +47,7 @@ struct FTestUserOp
 
 	struct Params
 	{
-		FAccountId LocalUserId;
+		FOnlineAccountIdHandle LocalUserId;
 	};
 
 	struct Result
@@ -74,7 +74,7 @@ struct FJoinableTestUserOp
 
 	struct Params
 	{
-		FAccountId LocalUserId;
+		FOnlineAccountIdHandle LocalUserId;
 	};
 
 	struct Result
@@ -109,7 +109,7 @@ struct FMergeableTestUserOp
 
 	struct Params
 	{
-		FAccountId LocalUserId;
+		FOnlineAccountIdHandle LocalUserId;
 		FTestMutation Mutations;
 	};
 
diff --git a/Engine/Plugins/Online/Experimental/OnlineServices/Source/OnlineServicesCommon/Public/Online/OnlineAsyncOpCache.h b/Engine/Plugins/Online/Experimental/OnlineServices/Source/OnlineServicesCommon/Public/Online/OnlineAsyncOpCache.h
--- a/Engine/Plugins/Online/Experimental/OnlineServices/Source/OnlineServicesCommon/Public/Online/OnlineAsyncOpCache.h
+++ b/Engine/Plugins/Online/Experimental/OnlineServices/Source/OnlineServicesCommon/Public/Online/OnlineAsyncOpCache.h
@@ -494,7 +494,7 @@ private:
 	friend uint32 GetTypeHash(const FWrappedOperationKey& Key);
 
 	TMap<FWrappedOperationKey, TUniquePtr<IWrappedOperation>> Operations;
-	TMap<FAccountId, TMap<FWrappedOperationKey, TUniquePtr<IWrappedOperation>>> UserOperations;
+	TMap<FOnlineAccountIdHandle, TMap<FWrappedOperationKey, TUniquePtr<IWrappedOperation>>> UserOperations;
 };
 
 inline uint32 GetTypeHash(const FOnlineAsyncOpCache::FWrappedOperationKey& Key)
diff --git a/Engine/Plugins/Online/Experimental/OnlineServices/Source/OnlineServicesCommon/Public/Online/OnlineComponent.h b/Engine/Plugins/Online/Experimental/OnlineServices/Source/OnlineServicesCommon/Public/Online/OnlineComponent.h
--- a/Engine/Plugins/Online/Experimental/OnlineServices/Source/OnlineServicesCommon/Public/Online/OnlineComponent.h
+++ b/Engine/Plugins/Online/Experimental/OnlineServices/Source/OnlineServicesCommon/Public/Online/OnlineComponent.h
@@ -78,7 +78,7 @@ public:
 	}
 
 	/* Queues for executing per-user tasks in serial */
-	FOnlineAsyncOpQueue& GetSerialQueue(FAccountId& AccountId)
+	FOnlineAsyncOpQueue& GetSerialQueue(const FOnlineAccountIdHandle& AccountId)
 	{
 		TUniquePtr<FOnlineAsyncOpQueueSerial>* Queue = PerUserSerialQueue.Find(AccountId);
 		if (Queue == nullptr)
@@ -177,7 +177,7 @@ private:
 	TMap<FString, TUniquePtr<IOnlineExecHandler>> ExecCommands;
 
 	FOnlineAsyncOpQueueSerial SerialQueue;
-	TMap<FAccountId, TUniquePtr<FOnlineAsyncOpQueueSerial>> PerUserSerialQueue;
+	TMap<FOnlineAccountIdHandle, TUniquePtr<FOnlineAsyncOpQueueSerial>> PerUserSerialQueue;
 };
 
 /* UE::Online */ }
diff --git a/Engine/Plugins/Online/Experimental/OnlineServices/Source/OnlineServicesCommon/Public/Online/OnlineIdCommon.h b/Engine/Plugins/Online/Experimental/OnlineServices/Source/OnlineServicesCommon/Public/Online/OnlineIdCommon.h
new file mode 100644
--- /dev/null
+++ b/Engine/Plugins/Online/Experimental/OnlineServices/Source/OnlineServicesCommon/Public/Online/OnlineIdCommon.h
@@ -0,0 +1,83 @@
+// Copyright Epic Games, Inc. All Rights Reserved.
+
+#pragma once
+
+#include "Misc/ScopeRWLock.h"
+
+namespace UE::Online {
+
+/**
+ * A net id registry suitable for use with trivial immutable keys
+ */
+template<EOnlineIdType IdType, typename IdValueType, EOnlineServices SubsystemType>
+class TOnlineBasicIdRegistry
+{
+public:
+	using HandleType = TOnlineIdHandle<IdType>;
+
+	static TOnlineBasicIdRegistry& Get();
+
+	HandleType FindOrAddHandle(const IdValueType& IdValue)
+	{
+		HandleType Handle;
+		Handle.Type = SubsystemType;
+
+		// Take a read lock and check if we already have a handle
+		{
+			FReadScopeLock ReadLock(Lock);
+			if (const uint32* FoundHandle = IdValueToHandleMap.Find(IdValue))
+			{
+				Handle.Handle = *FoundHandle;
+			}
+		}
+
+		if (!Handle.IsValid())
+		{
+			// Take a write lock, check again if we already have a handle, or insert a new element.
+			FWriteScopeLock WriteLock(Lock);
+			if (const HandleType* FoundHandle = IdValueToHandleMap.Find(IdValue))
+			{
+				Handle.Handle = *FoundHandle;
+			}
+
+			if(!Handle.IsValid())
+			{
+				IdValues.Emplace(IdValue);
+				Handle.Handle = IdValues.Num();
+				IdValueToHandleMap.Emplace(IdValue, Handle.Handle);
+			}
+		}
+
+		return Handle;
+	}
+
+	// Returns a copy as it's not thread safe to return a pointer/ref to an element of an array that can be relocated by another thread.
+	IdValueType GetIdValue(const HandleType Handle) const
+	{
+		if (Handle.Type == SubsystemType
+			&& Handle.IsValid()
+			&& IdValues.IsValidIndex(Handle.Handle - 1))
+		{
+			FReadScopeLock ReadLock(Lock);
+			return IdValues[Handle.Handle - 1];
+		}
+		return IdValueType();
+	}
+
+	IdValueType GetIdValueChecked(const HandleType Handle) const
+	{
+		check(Handle.Type == SubsystemType
+			&& Handle.IsValid()
+			&& IdValues.IsValidIndex(Handle.Handle - 1));
+		return GetIdValue(Handle);
+	}
+
+private:
+	mutable FRWLock Lock;
+
+	TArray<IdValueType> IdValues;
+	TMap<IdValueType, uint32> IdValueToHandleMap;
+};
+
+/* UE::Online */ }
+
diff --git a/Engine/Plugins/Online/Experimental/OnlineServices/Source/OnlineServicesCommon/Public/Online/OnlineServicesCommon.h b/Engine/Plugins/Online/Experimental/OnlineServices/Source/OnlineServicesCommon/Public/Online/OnlineServicesCommon.h
--- a/Engine/Plugins/Online/Experimental/OnlineServices/Source/OnlineServicesCommon/Public/Online/OnlineServicesCommon.h
+++ b/Engine/Plugins/Online/Experimental/OnlineServices/Source/OnlineServicesCommon/Public/Online/OnlineServicesCommon.h
@@ -35,7 +35,6 @@ public:
 	// IOnlineServices
 	virtual void Init() override;
 	virtual void Destroy() override;
-	virtual FAccountId CreateAccountId(FString&& InAccountIdString) override;
 	virtual IAuthPtr GetAuthInterface() override;
 	virtual IFriendsPtr GetFriendsInterface() override;
 	virtual IPresencePtr GetPresenceInterface() override;
@@ -300,7 +299,7 @@ public:
 	FOnlineAsyncOpQueue& GetSerialQueue();
 
 	/* Queues for executing per-user tasks in serial */
-	FOnlineAsyncOpQueue& GetSerialQueue(FAccountId& AccountId);
+	FOnlineAsyncOpQueue& GetSerialQueue(const FOnlineAccountIdHandle& AccountId);
 	
 	void RegisterExecHandler(const FString& Name, TUniquePtr<IOnlineExecHandler>&& Handler);
 
@@ -323,7 +322,7 @@ protected:
 
 	FOnlineAsyncOpQueueParallel ParallelQueue;
 	FOnlineAsyncOpQueueSerial SerialQueue;
-	TMap<FAccountId, TUniquePtr<FOnlineAsyncOpQueueSerial>> PerUserSerialQueue;
+	TMap<FOnlineAccountIdHandle, TUniquePtr<FOnlineAsyncOpQueueSerial>> PerUserSerialQueue;
 };
 
 /* UE::Online */ }
diff --git a/Engine/Plugins/Online/Experimental/OnlineServices/Source/OnlineServicesInterface/Private/Auth.cpp b/Engine/Plugins/Online/Experimental/OnlineServices/Source/OnlineServicesInterface/Private/Auth.cpp
new file mode 100644
--- /dev/null
+++ b/Engine/Plugins/Online/Experimental/OnlineServices/Source/OnlineServicesInterface/Private/Auth.cpp
@@ -0,0 +1,39 @@
+// Copyright Epic Games, Inc. All Rights Reserved.
+
+#include "Online/Auth.h"
+
+namespace UE::Online {
+
+const TCHAR* LexToString(ELoginStatus Status)
+{
+	switch (Status)
+	{
+	case ELoginStatus::UsingLocalProfile:	return TEXT("UsingLocalProfile");
+	case ELoginStatus::LoggedIn:			return TEXT("LoggedIn");
+	default:								checkNoEntry(); // Intentional fallthrough
+	case ELoginStatus::NotLoggedIn:			return TEXT("NotLoggedIn");
+	}
+}
+
+void LexFromString(ELoginStatus& OutStatus, const TCHAR* InStr)
+{
+	if (FCString::Stricmp(InStr, TEXT("LoggedIn")) == 0)
+	{
+		OutStatus = ELoginStatus::LoggedIn;
+	}
+	else if (FCString::Stricmp(InStr, TEXT("UsingLocalProfile")) == 0)
+	{
+		OutStatus = ELoginStatus::UsingLocalProfile;
+	}
+	else if (FCString::Stricmp(InStr, TEXT("NotLoggedIn")) == 0)
+	{
+		OutStatus = ELoginStatus::NotLoggedIn;
+	}
+	else
+	{
+		checkNoEntry();
+		OutStatus = ELoginStatus::NotLoggedIn;
+	}
+}
+
+/* UE::Online */ }
diff --git a/Engine/Plugins/Online/Experimental/OnlineServices/Source/OnlineServicesInterface/Public/Online/Auth.h b/Engine/Plugins/Online/Experimental/OnlineServices/Source/OnlineServicesInterface/Public/Online/Auth.h
--- a/Engine/Plugins/Online/Experimental/OnlineServices/Source/OnlineServicesInterface/Public/Online/Auth.h
+++ b/Engine/Plugins/Online/Experimental/OnlineServices/Source/OnlineServicesInterface/Public/Online/Auth.h
@@ -19,17 +19,8 @@ enum class ELoginStatus
 	/** Player has been validated by the platform specific authentication service */
 	LoggedIn
 };
-
-inline const TCHAR* LexToString(ELoginStatus LoginStatus)
-{
-	switch (LoginStatus)
-	{
-		case ELoginStatus::NotLoggedIn:			return TEXT("NoLoggedIn"); break;
-		case ELoginStatus::UsingLocalProfile:	return TEXT("UsingLocalProfile"); break;
-		case ELoginStatus::LoggedIn:			return TEXT("LoggedIn"); break;
-		default:								return TEXT("Unknown"); break;
-	}
-}
+ONLINESERVICESINTERFACE_API const TCHAR* LexToString(ELoginStatus Status);
+ONLINESERVICESINTERFACE_API void LexFromString(ELoginStatus& OutStatus, const TCHAR* InStr);
 
 class FAccountInfo
 {
@@ -37,7 +28,7 @@ public:
 	/** Local user num */
 	int32 LocalUserNum;
 	/** Account id */
-	FAccountId UserId;
+	FOnlineAccountIdHandle UserId;
 	/** Login status */
 	ELoginStatus LoginStatus;
 	// TODO: Other fields
@@ -68,7 +59,7 @@ struct FAuthLogout
 
 	struct Params
 	{
-		FAccountId LocalUserId;
+		FOnlineAccountIdHandle LocalUserId;
 		bool bDestroyAuth = false;
 	};
 
@@ -83,7 +74,7 @@ struct FAuthGenerateAuth
 
 	struct Params
 	{
-		FAccountId LocalUserId;
+		FOnlineAccountIdHandle LocalUserId;
 		FString Type;
 		TArray<FString> Scopes;
 	};
@@ -114,7 +105,7 @@ struct FAuthGetAccountByAccountId
 
 	struct Params
 	{
-		FAccountId LocalUserId;
+		FOnlineAccountIdHandle LocalUserId;
 	};
 
 	struct Result
@@ -127,7 +118,7 @@ struct FAuthGetAccountByAccountId
 struct FLoginStatusChanged
 {
 	/** User id whose status has changed */
-	FAccountId LocalUserId;
+	FOnlineAccountIdHandle LocalUserId;
 	/** Previous login status */
 	ELoginStatus PreviousStatus;
 	/** Current login status */
diff --git a/Engine/Plugins/Online/Experimental/OnlineServices/Source/OnlineServicesInterface/Public/Online/ExternalUI.h b/Engine/Plugins/Online/Experimental/OnlineServices/Source/OnlineServicesInterface/Public/Online/ExternalUI.h
--- a/Engine/Plugins/Online/Experimental/OnlineServices/Source/OnlineServicesInterface/Public/Online/ExternalUI.h
+++ b/Engine/Plugins/Online/Experimental/OnlineServices/Source/OnlineServicesInterface/Public/Online/ExternalUI.h
@@ -16,7 +16,7 @@ struct FExternalUIShowFriendsUI
 
 	struct Params
 	{
-		FAccountId AccountId;
+		FOnlineAccountIdHandle LocalUserId;
 	};
 
 	struct Result
@@ -37,7 +37,7 @@ namespace Meta {
 // TODO: Move to ExternalUI_Meta.inl file?
 
 BEGIN_ONLINE_STRUCT_META(FExternalUIShowFriendsUI::Params)
-	ONLINE_STRUCT_FIELD(FExternalUIShowFriendsUI::Params, AccountId)
+	ONLINE_STRUCT_FIELD(FExternalUIShowFriendsUI::Params, LocalUserId)
 END_ONLINE_STRUCT_META()
 
 BEGIN_ONLINE_STRUCT_META(FExternalUIShowFriendsUI::Result)
diff --git a/Engine/Plugins/Online/Experimental/OnlineServices/Source/OnlineServicesInterface/Public/Online/Friends.h b/Engine/Plugins/Online/Experimental/OnlineServices/Source/OnlineServicesInterface/Public/Online/Friends.h
--- a/Engine/Plugins/Online/Experimental/OnlineServices/Source/OnlineServicesInterface/Public/Online/Friends.h
+++ b/Engine/Plugins/Online/Experimental/OnlineServices/Source/OnlineServicesInterface/Public/Online/Friends.h
@@ -27,7 +27,7 @@ const TCHAR* LexToString(EFriendInviteStatus);
 struct FFriend
 {
 	/** Id of the friend */
-	FAccountId UserId;
+	FOnlineAccountIdHandle UserId;
 	/** Friendship invite status - if an invite is pending */
 	TOptional<EFriendInviteStatus> InviteStatus;
 };
@@ -40,7 +40,7 @@ struct FQueryFriends
 	struct Params
 	{
 		/** Account Id of the local user making the request */
-		FAccountId LocalUserId;
+		FOnlineAccountIdHandle LocalUserId;
 	};
 
 	/**
@@ -60,7 +60,7 @@ struct FGetFriends
 	struct Params
 	{
 		/** Account Id of the local user making the request */
-		FAccountId LocalUserId;
+		FOnlineAccountIdHandle LocalUserId;
 	};
 
 	/** Output struct for Friends::GetFriends */
@@ -79,9 +79,9 @@ struct FAddFriend
 	struct Params
 	{
 		/** Account Id of the local user making the request */
-		FAccountId LocalUserId;
+		FOnlineAccountIdHandle LocalUserId;
 		/** Friend to add */
-		FAccountId FriendId;
+		FOnlineAccountIdHandle FriendId;
 	};
 
 	/** Output struct for Friends::AddFriend */
@@ -96,13 +96,13 @@ struct FAddFriend
 struct FFriendsListUpdated
 {
 	/** Local user who's friends list changed */
-	FAccountId LocalUserId;
+	FOnlineAccountIdHandle LocalUserId;
 	/** Account ids of friends that have been added */
-	TArrayView<const FAccountId> AddedFriends;
+	TArrayView<const FOnlineAccountIdHandle> AddedFriends;
 	/** Account ids of friends that have been removed */
-	TArrayView<const FAccountId> RemovedFriends;
+	TArrayView<const FOnlineAccountIdHandle> RemovedFriends;
 	/** Account ids of friends that have been updated */
-	TArrayView<const FAccountId> UpdatedFriends;
+	TArrayView<const FOnlineAccountIdHandle> UpdatedFriends;
 };
 
 class IFriends
diff --git a/Engine/Plugins/Online/Experimental/OnlineServices/Source/OnlineServicesInterface/Public/Online/OnlineId.h b/Engine/Plugins/Online/Experimental/OnlineServices/Source/OnlineServicesInterface/Public/Online/OnlineId.h
--- a/Engine/Plugins/Online/Experimental/OnlineServices/Source/OnlineServicesInterface/Public/Online/OnlineId.h
+++ b/Engine/Plugins/Online/Experimental/OnlineServices/Source/OnlineServicesInterface/Public/Online/OnlineId.h
@@ -3,46 +3,61 @@
 #pragma once
 
 #include "Templates/TypeHash.h"
-#include "UObject/NameTypes.h"
 
-/**
- * This file is a placeholder, will be implemented along with the id registry in the future
- */
+#include "Online/OnlineServicesTypes.h"
 
 namespace UE::Online {
 
-using FOnlineIdHandle = int32;
-struct FOnlineId
+/** Enum used as template argument to TOnlineIdHandle to make it a compile error to assign between id's of different types */
+enum class EOnlineIdType
+{
+	AccountId,
+	SessionId,
+	PartyId
+};
+
+/**
+ * A handle to an id which uniquely identifies a persistent or transient online resource, i.e. account/session/party etc, within a given Online Services provider.
+ * At most one id, and therefore one handle, exists for any given resource. The id and handle persist until the OnlineServices module is unloaded.
+ * Passed to and returned from OnlineServices APIs.
+ */ 
+template<EOnlineIdType IdType>
+class TOnlineIdHandle
 {
-	FName Type;
-	FOnlineIdHandle Handle;
+public:
+	TOnlineIdHandle() : Type(EOnlineServices::Null), Handle(0) {}
+
+	inline bool IsValid() const { return Handle != 0; }
+
+	/* The Online Services provider this id relates to */
+	EOnlineServices Type : 8;
+	/* An opaque handle to the underlying online id */
+	uint32 Handle : 24;
 };
 
-using FAccountId = FOnlineId;
+using FOnlineAccountIdHandle = TOnlineIdHandle<EOnlineIdType::AccountId>;
 
-inline bool operator==(const FOnlineId& A, const FOnlineId& B)
+template<EOnlineIdType IdType>
+inline bool operator==(const TOnlineIdHandle<IdType>& A, const TOnlineIdHandle<IdType>& B)
 {
 	return A.Type == B.Type && A.Handle == B.Handle;
 }
 
-inline bool operator!=(const FOnlineId& A, const FOnlineId& B)
+template<EOnlineIdType IdType>
+inline bool operator!=(const TOnlineIdHandle<IdType>& A, const TOnlineIdHandle<IdType>& B)
 {
 	return !(A == B);
 }
 
-inline uint32 GetTypeHash(const FOnlineId& AccountId)
+template<EOnlineIdType IdType>
+inline uint32 GetTypeHash(const TOnlineIdHandle<IdType>& Handle)
 {
 	using ::GetTypeHash;
-	return HashCombine(GetTypeHash(AccountId.Type), GetTypeHash(AccountId.Handle));
-}
-
-inline FString ToLogString(const FOnlineId& Id)
-{
-	// TODO: Redact in shipping once we have the id registry version of this
-	return FString::Printf(TEXT("%s:%x"), *Id.Type.ToString(), Id.Handle);
+	return HashCombine(GetTypeHash(Handle.Type), GetTypeHash(Handle.Handle));
 }
 
-inline void LexFromString(FOnlineId& Id, const TCHAR* String)
+template<EOnlineIdType IdType>
+inline void LexFromString(const TOnlineIdHandle<IdType>& Id, const TCHAR* String)
 {
 	// TODO: should instead just implement ParseOnlineExecParams
 }
diff --git a/Engine/Plugins/Online/Experimental/OnlineServices/Source/OnlineServicesInterface/Public/Online/OnlineServices.h b/Engine/Plugins/Online/Experimental/OnlineServices/Source/OnlineServicesInterface/Public/Online/OnlineServices.h
--- a/Engine/Plugins/Online/Experimental/OnlineServices/Source/OnlineServicesInterface/Public/Online/OnlineServices.h
+++ b/Engine/Plugins/Online/Experimental/OnlineServices/Source/OnlineServicesInterface/Public/Online/OnlineServices.h
@@ -2,75 +2,18 @@
 
 #pragma once
 
+#include "Online/OnlineId.h"
 #include "Templates/SharedPointer.h"
 
 class FString;
 
 namespace UE::Online {
 
-enum class EOnlineServices : uint8
-{
-	// Null, Providing minimal functionality when no backend services are required
-	Null,
-	// Epic Online Services
-	Epic,
-	// Xbox services
-	Xbox,
-	// PlayStation Network
-	PSN,
-	// Nintendo
-	Nintendo,
-	// Stadia,
-	Stadia,
-	// Steam
-	Steam,
-	// Google
-	Google,
-	// GooglePlay
-	GooglePlay,
-	// Apple
-	Apple,
-	// GameKit
-	AppleGameKit,
-	// Samsung
-	Samsung,
-	// Oculus
-	Oculus,
-	// Tencent
-	Tencent,
-	// Reserved for future use/platform extensions
-	Reserved_14,
-	Reserved_15,
-	Reserved_16,
-	Reserved_17,
-	Reserved_18,
-	Reserved_19,
-	Reserved_20,
-	Reserved_21,
-	Reserved_22,
-	Reserved_23,
-	Reserved_24,
-	Reserved_25,
-	Reserved_26,
-	Reserved_27,
-	// For game specific Online Services
-	GameDefined_0 = 28,
-	GameDefined_1,
-	GameDefined_2,
-	GameDefined_3,
-	// Platform native, may not exist for all platforms
-	Platform = 254,
-	// Default, configured via ini, TODO: List specific ini section/key
-	Default = 255
-};
-
 // Interfaces
 using IAuthPtr = TSharedPtr<class IAuth>;
 using IFriendsPtr = TSharedPtr<class IFriends>;
 using IPresencePtr = TSharedPtr<class IPresence>;
 using IExternalUIPtr = TSharedPtr<class IExternalUI>;
-struct FOnlineId;
-using FAccountId = FOnlineId;
 
 class ONLINESERVICESINTERFACE_API IOnlineServices
 {
@@ -84,11 +27,6 @@ public:
 	 *
 	 */
 	virtual void Destroy() = 0;
-
-	/**
-	 *
-	 */
-	virtual FAccountId CreateAccountId(FString&& InAccountIdString) = 0;
 	
 	/**
 	 *
@@ -109,6 +47,11 @@ public:
 	 *
 	 */
 	virtual IExternalUIPtr GetExternalUIInterface() = 0;
+
+	/**
+	 * 
+	 */
+	virtual FString ToDebugString(const FOnlineAccountIdHandle& Handle) = 0;
 };
 
 /**
@@ -140,5 +83,15 @@ TSharedPtr<ServicesClass> GetServices(FName InstanceName = NAME_None)
  */
 ONLINESERVICESINTERFACE_API void DestroyServices(EOnlineServices OnlineServices = EOnlineServices::Default, FName InstanceName = NAME_None);
 
+template<EOnlineIdType IdType>
+inline FString ToLogString(const TOnlineIdHandle<IdType>& Id)
+{
+	FString Result;
+	if (TSharedPtr<IOnlineServices> Services = GetServices(Id.Type))
+	{
+		Result = Services->ToDebugString(Id);
+	}
+	return Result;
+}
 
 /* UE::Online */ }
diff --git a/Engine/Plugins/Online/Experimental/OnlineServices/Source/OnlineServicesInterface/Public/Online/OnlineServicesTypes.h b/Engine/Plugins/Online/Experimental/OnlineServices/Source/OnlineServicesInterface/Public/Online/OnlineServicesTypes.h
new file mode 100644
--- /dev/null
+++ b/Engine/Plugins/Online/Experimental/OnlineServices/Source/OnlineServicesInterface/Public/Online/OnlineServicesTypes.h
@@ -0,0 +1,63 @@
+// Copyright Epic Games, Inc. All Rights Reserved.
+
+#pragma once
+
+namespace UE::Online {
+
+enum class EOnlineServices : uint8
+{
+	// Null, Providing minimal functionality when no backend services are required
+	Null,
+	// Epic Online Services
+	Epic,
+	// Xbox services
+	Xbox,
+	// PlayStation Network
+	PSN,
+	// Nintendo
+	Nintendo,
+	// Stadia,
+	Stadia,
+	// Steam
+	Steam,
+	// Google
+	Google,
+	// GooglePlay
+	GooglePlay,
+	// Apple
+	Apple,
+	// GameKit
+	AppleGameKit,
+	// Samsung
+	Samsung,
+	// Oculus
+	Oculus,
+	// Tencent
+	Tencent,
+	// Reserved for future use/platform extensions
+	Reserved_14,
+	Reserved_15,
+	Reserved_16,
+	Reserved_17,
+	Reserved_18,
+	Reserved_19,
+	Reserved_20,
+	Reserved_21,
+	Reserved_22,
+	Reserved_23,
+	Reserved_24,
+	Reserved_25,
+	Reserved_26,
+	Reserved_27,
+	// For game specific Online Services
+	GameDefined_0 = 28,
+	GameDefined_1,
+	GameDefined_2,
+	GameDefined_3,
+	// Platform native, may not exist for all platforms
+	Platform = 254,
+	// Default, configured via ini, TODO: List specific ini section/key
+	Default = 255
+};
+
+/* UE::Online */ }
diff --git a/Engine/Plugins/Online/Experimental/OnlineServices/Source/OnlineServicesInterface/Public/Online/Presence.h b/Engine/Plugins/Online/Experimental/OnlineServices/Source/OnlineServicesInterface/Public/Online/Presence.h
--- a/Engine/Plugins/Online/Experimental/OnlineServices/Source/OnlineServicesInterface/Public/Online/Presence.h
+++ b/Engine/Plugins/Online/Experimental/OnlineServices/Source/OnlineServicesInterface/Public/Online/Presence.h
@@ -29,7 +29,7 @@ typedef TMap<FString, FPresenceVariant> FPresenceProperties;
 struct FUserPresence
 {
 	/** User whose presence this is */
-	FAccountId UserId;
+	FOnlineAccountIdHandle UserId;
 	/** Presence state */
 	EPresenceState State;
 	/** Presence string */
@@ -45,9 +45,9 @@ struct FQueryPresence
 	struct Params
 	{
 		/** Local user performing the query */
-		FAccountId LocalUserId;
+		FOnlineAccountIdHandle LocalUserId;
 		/** User to query the presence for */
-		FAccountId TargetUserId;
+		FOnlineAccountIdHandle TargetUserId;
 	};
 
 	struct Result
@@ -64,9 +64,9 @@ struct FGetPresence
 	struct Params
 	{
 		/** Local user getting the presence */
-		FAccountId LocalUserId;
+		FOnlineAccountIdHandle LocalUserId;
 		/** User to get the presence for */
-		FAccountId TargetUserId;
+		FOnlineAccountIdHandle TargetUserId;
 	};
 
 	struct Result
@@ -83,7 +83,7 @@ struct FUpdatePresence
 	struct Params
 	{
 		/** Local user performing the query */
-		FAccountId LocalUserId;
+		FOnlineAccountIdHandle LocalUserId;
 		/** Mutations */
 		struct FMutations
 		{
@@ -111,7 +111,7 @@ struct FUpdatePresence
 struct FPresenceUpdated
 {
 	/** Local user receiving the presence update */
-	FAccountId LocalUserId;
+	FOnlineAccountIdHandle LocalUserId;
 	/** Presence that has updated */
 	TSharedRef<const FUserPresence> UpdatedPresence;
 };
diff --git a/Engine/Plugins/Online/Experimental/OnlineServicesEOS/Source/Private/AuthEOS.cpp b/Engine/Plugins/Online/Experimental/OnlineServicesEOS/Source/Private/AuthEOS.cpp
--- a/Engine/Plugins/Online/Experimental/OnlineServicesEOS/Source/Private/AuthEOS.cpp
+++ b/Engine/Plugins/Online/Experimental/OnlineServicesEOS/Source/Private/AuthEOS.cpp
@@ -2,12 +2,17 @@
 
 #include "AuthEOS.h"
 
+#include "OnlineIdEOS.h"
 #include "OnlineServicesEOS.h"
 #include "OnlineServicesEOSTypes.h"
+#include "Algo/Transform.h"
 #include "Online/AuthErrors.h"
+#include "Online/OnlineAsyncOp.h"
 #include "Online/OnlineErrorDefinitions.h"
 
 #include "eos_auth.h"
+#include "eos_common.h"
+#include "eos_connect.h"
 #include "eos_types.h"
 #include "eos_init.h"
 #include "eos_sdk.h"
@@ -15,9 +20,6 @@
 
 namespace UE::Online {
 
-// TEMP until Net Id Registry is done
-TMap<EOS_EpicAccountId, int32> EOSAccountIdMap;
-
 static inline ELoginStatus ToELoginStatus(EOS_ELoginStatus InStatus)
 {
 	switch (InStatus)
@@ -119,10 +121,15 @@ void FAuthEOS::Initialize()
 	NotifyLoginStatusChangedNotificationId = EOS_Auth_AddNotifyLoginStatusChanged(AuthHandle, &Options, this, [](const EOS_Auth_LoginStatusChangedCallbackInfo* Data)
 	{
 		FAuthEOS* This = reinterpret_cast<FAuthEOS*>(Data->ClientData);
-		FAccountId LocalUserId = MakeEOSAccountId(Data->LocalUserId);
-		ELoginStatus PreviousStatus = ToELoginStatus(Data->PrevStatus);
-		ELoginStatus CurrentStatus = ToELoginStatus(Data->CurrentStatus);
-		This->OnEOSLoginStatusChanged(LocalUserId, PreviousStatus, CurrentStatus);
+
+		FOnlineAccountIdHandle LocalUserId = FindAccountId(Data->LocalUserId);
+		// invalid handle is expected for players logging in because this callback is called _before_ the login complete callback
+		if (LocalUserId.IsValid())
+		{
+			ELoginStatus PreviousStatus = ToELoginStatus(Data->PrevStatus);
+			ELoginStatus CurrentStatus = ToELoginStatus(Data->CurrentStatus);
+			This->OnEOSLoginStatusChanged(LocalUserId, PreviousStatus, CurrentStatus);
+		}
 	});
 }
 
@@ -280,7 +287,7 @@ TOnlineAsyncOpHandle<FAuthLogin> FAuthEOS::Login(FAuthLogin::Params&& Params)
 	Op->Then([this, LoginOptions, Credentials](TOnlineAsyncOp<FAuthLogin>& InAsyncOp) mutable
 	{
 		LoginOptions.Credentials = &Credentials;
-		return EOS_Async<EOS_Auth_LoginCallbackInfo>(InAsyncOp, EOS_Auth_Login, AuthHandle, LoginOptions);
+		return EOS_Async<EOS_Auth_LoginCallbackInfo>(EOS_Auth_Login, AuthHandle, LoginOptions);
 	})
 	.Then([this](TOnlineAsyncOp<FAuthLogin>& InAsyncOp, const EOS_Auth_LoginCallbackInfo* Data)
 	{
@@ -311,7 +318,7 @@ TOnlineAsyncOpHandle<FAuthLogin> FAuthEOS::Login(FAuthLogin::Params&& Params)
 				ConnectLoginOptions.ApiVersion = EOS_CONNECT_LOGIN_API_LATEST;
 				ConnectLoginOptions.Credentials = &ConnectLoginCredentials;
 
-				return EOS_Async<EOS_Connect_LoginCallbackInfo>(InAsyncOp, EOS_Connect_Login, ConnectHandle, ConnectLoginOptions);
+				return EOS_Async<EOS_Connect_LoginCallbackInfo>(EOS_Connect_Login, ConnectHandle, ConnectLoginOptions);
 			}
 			else
 			{
@@ -354,7 +361,7 @@ TOnlineAsyncOpHandle<FAuthLogin> FAuthEOS::Login(FAuthLogin::Params&& Params)
 			ConnectCreateUserOptions.ApiVersion = EOS_CONNECT_CREATEUSER_API_LATEST;
 			ConnectCreateUserOptions.ContinuanceToken = Data->ContinuanceToken;
 
-			return EOS_Async<EOS_Connect_CreateUserCallbackInfo>(InAsyncOp, EOS_Connect_CreateUser, ConnectHandle, ConnectCreateUserOptions);
+			return EOS_Async<EOS_Connect_CreateUserCallbackInfo>(EOS_Connect_CreateUser, ConnectHandle, ConnectCreateUserOptions);
 		}
 		else
 		{
@@ -390,24 +397,25 @@ TOnlineAsyncOpHandle<FAuthLogin> FAuthEOS::Login(FAuthLogin::Params&& Params)
 
 void FAuthEOS::ProcessSuccessfulLogin(TOnlineAsyncOp<FAuthLogin>& InAsyncOp)
 {
-	EOS_EpicAccountId EpicAccountId = *InAsyncOp.Data.Get<EOS_EpicAccountId>(TEXT("EpicAccountId"));
-	EOS_ProductUserId ProductUserId = *InAsyncOp.Data.Get<EOS_ProductUserId>(TEXT("ProductUserId"));
+	const EOS_EpicAccountId EpicAccountId = *InAsyncOp.Data.Get<EOS_EpicAccountId>(TEXT("EpicAccountId"));
+	const EOS_ProductUserId ProductUserId = *InAsyncOp.Data.Get<EOS_ProductUserId>(TEXT("ProductUserId"));
+	const FOnlineAccountIdHandle LocalUserId = CreateAccountId(EpicAccountId, ProductUserId);
 
-	UE_LOG(LogTemp, Verbose, TEXT("[FAuthEOS::Login] Successfully logged in as [%s/%s]"), *LexToString(EpicAccountId), *LexToString(ProductUserId));
+	UE_LOG(LogTemp, Verbose, TEXT("[FAuthEOS::Login] Successfully logged in as [%s]"), *ToLogString(LocalUserId));
 
 	TSharedRef<FAccountInfoEOS> AccountInfo = MakeShared<FAccountInfoEOS>();
 	AccountInfo->LocalUserNum = InAsyncOp.GetParams().LocalUserNum;
-	AccountInfo->UserId = MakeEOSAccountId(EpicAccountId);  // ProductUserId will be used here too, after the net id registry changes
+	AccountInfo->UserId = LocalUserId;
 
-	check(!AccountInfos.Contains(AccountInfo->UserId));
-	AccountInfos.Emplace(AccountInfo->UserId, AccountInfo);
+	check(!AccountInfos.Contains(LocalUserId));
+	AccountInfos.Emplace(LocalUserId, AccountInfo);
 
 	FAuthLogin::Result Result = { AccountInfo };
 	InAsyncOp.SetResult(MoveTemp(Result));
 
 	// Trigger event
 	FLoginStatusChanged EventParameters;
-	EventParameters.LocalUserId = AccountInfo->UserId;
+	EventParameters.LocalUserId = LocalUserId;
 	EventParameters.PreviousStatus = ELoginStatus::NotLoggedIn;
 	EventParameters.CurrentStatus = ELoginStatus::LoggedIn;
 	OnLoginStatusChangedEvent.Broadcast(EventParameters);
@@ -415,58 +423,65 @@ void FAuthEOS::ProcessSuccessfulLogin(TOnlineAsyncOp<FAuthLogin>& InAsyncOp)
 
 TOnlineAsyncOpHandle<FAuthLogout> FAuthEOS::Logout(FAuthLogout::Params&& Params)
 {
-	FAccountId ParamLocalUserId = Params.LocalUserId;
-	TOptional<EOS_EpicAccountId> AccountId = EOSAccountIdFromOnlineServiceAccountId(ParamLocalUserId);
+	const FOnlineAccountIdHandle LocalUserId = Params.LocalUserId;
 	TOnlineAsyncOpRef<FAuthLogout> Op = GetOp<FAuthLogout>(MoveTemp(Params));
-	if (AccountId.IsSet() && AccountInfos.Contains(ParamLocalUserId))
-	{
-		// Should we destroy persistent auth first?
-		TOnlineChainableAsyncOp<FAuthLogout, void> NextOp = *Op;
-		if (Params.bDestroyAuth)
-		{
-			EOS_Auth_DeletePersistentAuthOptions DeletePersistentAuthOptions = {0};
-			DeletePersistentAuthOptions.ApiVersion = EOS_AUTH_DELETEPERSISTENTAUTH_API_LATEST;
-			DeletePersistentAuthOptions.RefreshToken = nullptr; // Is this needed?  Docs say it's needed for consoles
-			NextOp = NextOp.Then([this, DeletePersistentAuthOptions](TOnlineAsyncOp<FAuthLogout>& InAsyncOp)
-				{
-					return EOS_Async<EOS_Auth_DeletePersistentAuthCallbackInfo>(InAsyncOp, EOS_Auth_DeletePersistentAuth, AuthHandle, DeletePersistentAuthOptions);
-				})
-				.Then([](TOnlineAsyncOp<FAuthLogout>& InAsyncOp, const EOS_Auth_DeletePersistentAuthCallbackInfo* Data)
-				{
-					UE_LOG(LogTemp, Warning, TEXT("DeletePersistentAuthResult: [%s]"), ANSI_TO_TCHAR(EOS_EResult_ToString(Data->ResultCode)));
-					// Regardless of success/failure, continue
-				});
-		}
-		// Logout
-		NextOp.Then([this, AccountId](TOnlineAsyncOp<FAuthLogout>& InAsyncOp)
-			{
-				EOS_Auth_LogoutOptions LogoutOptions = { };
-				LogoutOptions.ApiVersion = EOS_AUTH_LOGOUT_API_LATEST;
-				LogoutOptions.LocalUserId = AccountId.GetValue();
-				return EOS_Async<EOS_Auth_LogoutCallbackInfo>(InAsyncOp, EOS_Auth_Logout, AuthHandle, LogoutOptions);
-			})
-			.Then([](TOnlineAsyncOp<FAuthLogout>& InAsyncOp, const EOS_Auth_LogoutCallbackInfo* Data)
-			{
-				UE_LOG(LogTemp, Warning, TEXT("LogoutResult: [%s]"), *LexToString(Data->ResultCode));
 
-				if (Data->ResultCode == EOS_EResult::EOS_Success)
-				{
-					// Success
-					InAsyncOp.SetResult(FAuthLogout::Result());
-				}
-				else
-				{
-					// TODO: Error codes
-					FOnlineError Error = Errors::Unknown();
-					InAsyncOp.SetError(MoveTemp(Error));
-				}
-			}).Enqueue(GetSerialQueue());
+	if (!ValidateOnlineId(LocalUserId))
+	{
+		Op->SetError(Errors::InvalidUser());
+		return Op->GetHandle();
 	}
-	else
+
+	const EOS_EpicAccountId LocalUserEasId = GetEpicAccountId(Params.LocalUserId);
+	if (!EOS_EpicAccountId_IsValid(LocalUserEasId) || !AccountInfos.Contains(LocalUserId))
 	{
 		// TODO: Error codes
 		Op->SetError(Errors::Unknown());
+		return Op->GetHandle();
 	}
+
+	// Should we destroy persistent auth first?
+	TOnlineChainableAsyncOp<FAuthLogout, void> NextOp = *Op;
+	if (Params.bDestroyAuth)
+	{
+		EOS_Auth_DeletePersistentAuthOptions DeletePersistentAuthOptions = {0};
+		DeletePersistentAuthOptions.ApiVersion = EOS_AUTH_DELETEPERSISTENTAUTH_API_LATEST;
+		DeletePersistentAuthOptions.RefreshToken = nullptr; // Is this needed?  Docs say it's needed for consoles
+		NextOp = NextOp.Then([this, DeletePersistentAuthOptions](TOnlineAsyncOp<FAuthLogout>& InAsyncOp)
+		{
+			return EOS_Async<EOS_Auth_DeletePersistentAuthCallbackInfo>(EOS_Auth_DeletePersistentAuth, AuthHandle, DeletePersistentAuthOptions);
+		})
+		.Then([](TOnlineAsyncOp<FAuthLogout>& InAsyncOp, const EOS_Auth_DeletePersistentAuthCallbackInfo* Data)
+		{
+			UE_LOG(LogTemp, Warning, TEXT("DeletePersistentAuthResult: [%s]"), ANSI_TO_TCHAR(EOS_EResult_ToString(Data->ResultCode)));
+			// Regardless of success/failure, continue
+		});
+	}
+
+	// Logout
+	NextOp.Then([this, LocalUserEasId](TOnlineAsyncOp<FAuthLogout>& InAsyncOp)
+	{
+		EOS_Auth_LogoutOptions LogoutOptions = { };
+		LogoutOptions.ApiVersion = EOS_AUTH_LOGOUT_API_LATEST;
+		LogoutOptions.LocalUserId = LocalUserEasId;
+		return EOS_Async<EOS_Auth_LogoutCallbackInfo>(EOS_Auth_Logout, AuthHandle, LogoutOptions);
+	})
+	.Then([](TOnlineAsyncOp<FAuthLogout>& InAsyncOp, const EOS_Auth_LogoutCallbackInfo* Data)
+	{
+		UE_LOG(LogTemp, Warning, TEXT("LogoutResult: [%s]"), *LexToString(Data->ResultCode));
+
+		if (Data->ResultCode == EOS_EResult::EOS_Success)
+		{
+			// Success
+			InAsyncOp.SetResult(FAuthLogout::Result());
+		}
+		else
+		{
+			// TODO: Error codes
+			FOnlineError Error = Errors::Unknown();
+			InAsyncOp.SetError(MoveTemp(Error));
+		}
+	}).Enqueue(GetSerialQueue());
 	
 	return Op->GetHandle();
 }
@@ -479,7 +494,7 @@ TOnlineAsyncOpHandle<FAuthGenerateAuth> FAuthEOS::GenerateAuth(FAuthGenerateAuth
 
 TOnlineResult<FAuthGetAccountByLocalUserNum> FAuthEOS::GetAccountByLocalUserNum(FAuthGetAccountByLocalUserNum::Params&& Params)
 {
-	TResult<FAccountId, FOnlineError> LocalUserIdResult = GetAccountIdByLocalUserNum(Params.LocalUserNum);
+	TResult<FOnlineAccountIdHandle, FOnlineError> LocalUserIdResult = GetAccountIdByLocalUserNum(Params.LocalUserNum);
 	if (LocalUserIdResult.IsOk())
 	{
 		FAuthGetAccountByLocalUserNum::Result Result = { AccountInfos.FindChecked(LocalUserIdResult.GetOkValue()) };
@@ -506,32 +521,29 @@ TOnlineResult<FAuthGetAccountByAccountId> FAuthEOS::GetAccountByAccountId(FAuthG
 	}
 }
 
-bool FAuthEOS::IsLoggedIn(const FAccountId& AccountId) const
+bool FAuthEOS::IsLoggedIn(const FOnlineAccountIdHandle& AccountId) const
 {
 	// TODO:  More logic?
 	return AccountInfos.Contains(AccountId);
 }
 
-TResult<FAccountId, FOnlineError> FAuthEOS::GetAccountIdByLocalUserNum(int32 LocalUserNum) const
+TResult<FOnlineAccountIdHandle, FOnlineError> FAuthEOS::GetAccountIdByLocalUserNum(int32 LocalUserNum) const
 {
-	for (const TPair<FAccountId, TSharedRef<FAccountInfoEOS>>& AccountPair : AccountInfos)
+	for (const TPair<FOnlineAccountIdHandle, TSharedRef<FAccountInfoEOS>>& AccountPair : AccountInfos)
 	{
 		if (AccountPair.Value->LocalUserNum == LocalUserNum)
 		{
-			TResult<FAccountId, FOnlineError> Result(AccountPair.Key);
+			TResult<FOnlineAccountIdHandle, FOnlineError> Result(AccountPair.Key);
 			return Result;
 		}
 	}
-	TResult<FAccountId, FOnlineError> Result(Errors::Unknown()); // TODO: error code
+	TResult<FOnlineAccountIdHandle, FOnlineError> Result(Errors::Unknown()); // TODO: error code
 	return Result;
 }
 
-void FAuthEOS::OnEOSLoginStatusChanged(FAccountId LocalUserId, ELoginStatus PreviousStatus, ELoginStatus CurrentStatus)
+void FAuthEOS::OnEOSLoginStatusChanged(FOnlineAccountIdHandle LocalUserId, ELoginStatus PreviousStatus, ELoginStatus CurrentStatus)
 {
-	UE_LOG(LogTemp, Warning, TEXT("OnEOSLoginStatusChanged: [%s] %d %d"),
-		*LexToString(EOSAccountIdFromOnlineServiceAccountId(LocalUserId).GetValue()),
-		(int)PreviousStatus,
-		(int)CurrentStatus);
+	UE_LOG(LogTemp, Warning, TEXT("OnEOSLoginStatusChanged: [%s] [%s]->[%s]"), *ToLogString(LocalUserId), LexToString(PreviousStatus), LexToString(CurrentStatus));
 	if (TSharedRef<FAccountInfoEOS>* AccountInfoPtr = AccountInfos.Find(LocalUserId))
 	{
 		FAccountInfoEOS& AccountInfo = **AccountInfoPtr;
@@ -555,5 +567,256 @@ void FAuthEOS::OnEOSLoginStatusChanged(FAccountId LocalUserId, ELoginStatus Prev
 	}
 }
 
+template <typename IdType>
+TFuture<FOnlineAccountIdHandle> ResolveAccountIdImpl(FAuthEOS& AuthEOS, const FOnlineAccountIdHandle& LocalUserId, const IdType InId)
+{
+	TPromise<FOnlineAccountIdHandle> Promise;
+	TFuture<FOnlineAccountIdHandle> Future = Promise.GetFuture();
+
+	AuthEOS.ResolveAccountIds(LocalUserId, { InId }).Next([Promise = MoveTemp(Promise)](const TArray<FOnlineAccountIdHandle>& AccountIds) mutable
+	{
+		FOnlineAccountIdHandle Result;
+		if (AccountIds.Num() == 1)
+		{
+			Result = AccountIds[0];
+		}
+		Promise.SetValue(Result);
+	});
+
+	return Future;
+}
+
+TFuture<FOnlineAccountIdHandle> FAuthEOS::ResolveAccountId(const FOnlineAccountIdHandle& LocalUserId, const EOS_EpicAccountId EpicAccountId)
+{
+	return ResolveAccountIdImpl(*this, LocalUserId, EpicAccountId);
+}
+
+TFuture<FOnlineAccountIdHandle> FAuthEOS::ResolveAccountId(const FOnlineAccountIdHandle& LocalUserId, const EOS_ProductUserId ProductUserId)
+{
+	return ResolveAccountIdImpl(*this, LocalUserId, ProductUserId);
+}
+
+using FEpicAccountIdStrBuffer = char[EOS_EPICACCOUNTID_MAX_LENGTH + 1];
+
+TFuture<TArray<FOnlineAccountIdHandle>> FAuthEOS::ResolveAccountIds(const FOnlineAccountIdHandle& LocalUserId, const TArray<EOS_EpicAccountId>& InEpicAccountIds)
+{
+	// Search for all the account id's
+	TArray<FOnlineAccountIdHandle> AccountIdHandles;
+	AccountIdHandles.Reserve(InEpicAccountIds.Num());
+	TArray<EOS_EpicAccountId> MissingEpicAccountIds;
+	MissingEpicAccountIds.Reserve(InEpicAccountIds.Num());
+	for (const EOS_EpicAccountId EpicAccountId : InEpicAccountIds)
+	{
+		if (!EOS_EpicAccountId_IsValid(EpicAccountId))
+		{
+			return MakeFulfilledPromise<TArray<FOnlineAccountIdHandle>>().GetFuture();
+		}
+
+		FOnlineAccountIdHandle Found = FindAccountId(EpicAccountId);
+		if (!Found.IsValid())
+		{
+			MissingEpicAccountIds.Emplace(EpicAccountId);
+		}
+		AccountIdHandles.Emplace(MoveTemp(Found));
+	}
+	if (MissingEpicAccountIds.IsEmpty())
+	{
+		// We have them all, so we can just return
+		return MakeFulfilledPromise<TArray<FOnlineAccountIdHandle>>(MoveTemp(AccountIdHandles)).GetFuture();
+	}
+
+	// If we failed to find all the handles, we need to query, which requires a valid LocalUserId
+	if (!ValidateOnlineId(LocalUserId))
+	{
+		checkNoEntry();
+		return MakeFulfilledPromise<TArray<FOnlineAccountIdHandle>>().GetFuture();
+	}
+
+	TPromise<TArray<FOnlineAccountIdHandle>> Promise;
+	TFuture<TArray<FOnlineAccountIdHandle>> Future = Promise.GetFuture();
+
+	TArray<FEpicAccountIdStrBuffer> EpicAccountIdStrsToQuery;
+	EpicAccountIdStrsToQuery.Reserve(MissingEpicAccountIds.Num());
+	for (const EOS_EpicAccountId EpicAccountId : MissingEpicAccountIds)
+	{
+		FEpicAccountIdStrBuffer& EpicAccountIdStr = EpicAccountIdStrsToQuery.Emplace_GetRef();
+		int32_t BufferSize = sizeof(EpicAccountIdStr);
+		if (!EOS_EpicAccountId_IsValid(EpicAccountId) ||
+			EOS_EpicAccountId_ToString(EpicAccountId, EpicAccountIdStr, &BufferSize) != EOS_EResult::EOS_Success)
+		{
+			checkNoEntry();
+			return MakeFulfilledPromise<TArray<FOnlineAccountIdHandle>>().GetFuture();
+		}
+	}
+
+	TArray<const char*> EpicAccountIdStrPtrs;
+	Algo::Transform(EpicAccountIdStrsToQuery, EpicAccountIdStrPtrs, [](const FEpicAccountIdStrBuffer& Str) { return &Str[0]; });
+
+	EOS_Connect_QueryExternalAccountMappingsOptions Options = {};
+	Options.ApiVersion = EOS_CONNECT_QUERYEXTERNALACCOUNTMAPPINGS_API_LATEST;
+	Options.LocalUserId = GetProductUserIdChecked(LocalUserId);
+	Options.AccountIdType = EOS_EExternalAccountType::EOS_EAT_EPIC;
+	Options.ExternalAccountIds = (const char**)EpicAccountIdStrPtrs.GetData();
+	Options.ExternalAccountIdCount = 1;
+
+	EOS_Async<EOS_Connect_QueryExternalAccountMappingsCallbackInfo>(EOS_Connect_QueryExternalAccountMappings, ConnectHandle, Options)
+		.Next([this, InEpicAccountIds, Promise = MoveTemp(Promise)](const EOS_Connect_QueryExternalAccountMappingsCallbackInfo* Data) mutable
+	{
+		TArray<FOnlineAccountIdHandle> AccountIds;
+		AccountIds.Reserve(InEpicAccountIds.Num());
+		if (Data->ResultCode == EOS_EResult::EOS_Success)
+		{
+			EOS_Connect_GetExternalAccountMappingsOptions Options = {};
+			Options.ApiVersion = EOS_CONNECT_GETEXTERNALACCOUNTMAPPING_API_LATEST;
+			Options.LocalUserId = Data->LocalUserId;
+			Options.AccountIdType = EOS_EExternalAccountType::EOS_EAT_EPIC;
+
+			for (const EOS_EpicAccountId EpicAccountId : InEpicAccountIds)
+			{
+				FOnlineAccountIdHandle AccountId = FindAccountId(EpicAccountId);
+				if (!AccountId.IsValid())
+				{
+					FEpicAccountIdStrBuffer EpicAccountIdStr;
+					int32_t BufferSize = sizeof(EpicAccountIdStr);
+					verify(EOS_EpicAccountId_ToString(EpicAccountId, EpicAccountIdStr, &BufferSize) == EOS_EResult::EOS_Success);
+					Options.TargetExternalUserId = &EpicAccountIdStr[0];
+					const EOS_ProductUserId ProductUserId = EOS_Connect_GetExternalAccountMapping(ConnectHandle, &Options);
+					AccountId = CreateAccountId(EpicAccountId, ProductUserId);
+				}
+				AccountIds.Emplace(MoveTemp(AccountId));
+			}
+		}
+		else
+		{
+			UE_LOG(LogTemp, Warning, TEXT("ResolveAccountId failed to query external mapping Result=[%s]"), *LexToString(Data->ResultCode));
+		}
+		Promise.SetValue(MoveTemp(AccountIds));
+	});
+
+	return Future;
+}
+
+TFuture<TArray<FOnlineAccountIdHandle>> FAuthEOS::ResolveAccountIds(const FOnlineAccountIdHandle& LocalUserId, const TArray<EOS_ProductUserId>& InProductUserIds)
+{
+	// Search for all the account id's
+	TArray<FOnlineAccountIdHandle> AccountIdHandles;
+	AccountIdHandles.Reserve(InProductUserIds.Num());
+	TArray<EOS_ProductUserId> MissingProductUserIds;
+	MissingProductUserIds.Reserve(InProductUserIds.Num());
+	for (const EOS_ProductUserId ProductUserId : InProductUserIds)
+	{
+		if (!EOS_ProductUserId_IsValid(ProductUserId))
+		{
+			return MakeFulfilledPromise<TArray<FOnlineAccountIdHandle>>().GetFuture();
+		}
+
+		FOnlineAccountIdHandle Found = FindAccountId(ProductUserId);
+		if (!Found.IsValid())
+		{
+			MissingProductUserIds.Emplace(ProductUserId);
+		}
+		AccountIdHandles.Emplace(MoveTemp(Found));
+	}
+	if (MissingProductUserIds.IsEmpty())
+	{
+		// We have them all, so we can just return
+		return MakeFulfilledPromise<TArray<FOnlineAccountIdHandle>>(MoveTemp(AccountIdHandles)).GetFuture();
+	}
+
+	// If we failed to find all the handles, we need to query, which requires a valid LocalUserId
+	if (!ValidateOnlineId(LocalUserId))
+	{
+		checkNoEntry();
+		return MakeFulfilledPromise<TArray<FOnlineAccountIdHandle>>().GetFuture();
+	}
+
+	TPromise<TArray<FOnlineAccountIdHandle>> Promise;
+	TFuture<TArray<FOnlineAccountIdHandle>> Future = Promise.GetFuture();
+
+	EOS_Connect_QueryProductUserIdMappingsOptions Options = {};
+	Options.ApiVersion = EOS_CONNECT_QUERYPRODUCTUSERIDMAPPINGS_API_LATEST;
+	Options.LocalUserId = GetProductUserIdChecked(LocalUserId);
+	Options.ProductUserIds = MissingProductUserIds.GetData();
+	Options.ProductUserIdCount = MissingProductUserIds.Num();
+	EOS_Async<EOS_Connect_QueryProductUserIdMappingsCallbackInfo>(EOS_Connect_QueryProductUserIdMappings, ConnectHandle, Options)
+		.Next([this, InProductUserIds, Promise = MoveTemp(Promise)](const EOS_Connect_QueryProductUserIdMappingsCallbackInfo* Data) mutable
+	{
+		TArray<FOnlineAccountIdHandle> AccountIds;
+		if (Data->ResultCode == EOS_EResult::EOS_Success)
+		{
+			EOS_Connect_GetProductUserIdMappingOptions Options = {};
+			Options.ApiVersion = EOS_CONNECT_GETPRODUCTUSERIDMAPPING_API_LATEST;
+			Options.LocalUserId = Data->LocalUserId;
+			Options.AccountIdType = EOS_EExternalAccountType::EOS_EAT_EPIC;
+
+			for (const EOS_ProductUserId ProductUserId : InProductUserIds)
+			{
+				FOnlineAccountIdHandle AccountId = FindAccountId(ProductUserId);
+				if (!AccountId.IsValid())
+				{
+					Options.TargetProductUserId = ProductUserId;
+					FEpicAccountIdStrBuffer EpicAccountIdStr;
+					int32_t BufferLength = sizeof(EpicAccountIdStr);
+					verify(EOS_Connect_GetProductUserIdMapping(ConnectHandle, &Options, EpicAccountIdStr, &BufferLength) == EOS_EResult::EOS_Success);
+					const EOS_EpicAccountId EpicAccountId = EOS_EpicAccountId_FromString(EpicAccountIdStr);
+					check(EOS_EpicAccountId_IsValid(EpicAccountId));
+					AccountId = CreateAccountId(EpicAccountId, ProductUserId);
+				}
+				AccountIds.Emplace(MoveTemp(AccountId));
+			}
+		}
+		else
+		{
+			UE_LOG(LogTemp, Warning, TEXT("ResolveAccountId failed to query external mapping Result=[%s]"), *LexToString(Data->ResultCode));
+		}
+		Promise.SetValue(MoveTemp(AccountIds));
+	});
+
+	return Future;
+}
+
+template<typename ParamType>
+TFunction<TFuture<FOnlineAccountIdHandle>(FOnlineAsyncOp& InAsyncOp, const ParamType& Param)> ResolveIdFnImpl(FAuthEOS* AuthEOS)
+{
+	return [AuthEOS](FOnlineAsyncOp& InAsyncOp, const ParamType& Param)
+	{
+		const FOnlineAccountIdHandle* LocalUserIdPtr = InAsyncOp.Data.Get<FOnlineAccountIdHandle>(TEXT("LocalUserId"));
+		if (!ensure(LocalUserIdPtr))
+		{
+			return MakeFulfilledPromise<FOnlineAccountIdHandle>().GetFuture();
+		}
+		return AuthEOS->ResolveAccountId(*LocalUserIdPtr, Param);
+	};
+}
+TFunction<TFuture<FOnlineAccountIdHandle>(FOnlineAsyncOp& InAsyncOp, const EOS_EpicAccountId& ProductUserId)> FAuthEOS::ResolveEpicIdFn()
+{
+	return ResolveIdFnImpl<EOS_EpicAccountId>(this);
+}
+TFunction<TFuture<FOnlineAccountIdHandle>(FOnlineAsyncOp& InAsyncOp, const EOS_ProductUserId& ProductUserId)> FAuthEOS::ResolveProductIdFn()
+{
+	return ResolveIdFnImpl<EOS_ProductUserId>(this);
+}
+
+template<typename ParamType>
+TFunction<TFuture<TArray<FOnlineAccountIdHandle>>(FOnlineAsyncOp& InAsyncOp, const TArray<ParamType>& Param)> ResolveIdsFnImpl(FAuthEOS* AuthEOS)
+{
+	return [AuthEOS](FOnlineAsyncOp& InAsyncOp, const TArray<ParamType>& Param)
+	{
+		const FOnlineAccountIdHandle* LocalUserIdPtr = InAsyncOp.Data.Get<FOnlineAccountIdHandle>(TEXT("LocalUserId"));
+		if (!ensure(LocalUserIdPtr))
+		{
+			return MakeFulfilledPromise<TArray<FOnlineAccountIdHandle>>().GetFuture();
+		}
+		return AuthEOS->ResolveAccountIds(*LocalUserIdPtr, Param);
+	};
+}
+TFunction<TFuture<TArray<FOnlineAccountIdHandle>>(FOnlineAsyncOp& InAsyncOp, const TArray<EOS_EpicAccountId>& EpicAccountIds)> FAuthEOS::ResolveEpicIdsFn()
+{
+	return ResolveIdsFnImpl<EOS_EpicAccountId>(this);
+}
+TFunction<TFuture<TArray<FOnlineAccountIdHandle>>(FOnlineAsyncOp& InAsyncOp, const TArray<EOS_ProductUserId>& ProductUserIds)> FAuthEOS::ResolveProductIdsFn()
+{
+	return ResolveIdsFnImpl<EOS_ProductUserId>(this);
+}
 
 /* UE::Online */ }
diff --git a/Engine/Plugins/Online/Experimental/OnlineServicesEOS/Source/Private/ExternalUIEOS.cpp b/Engine/Plugins/Online/Experimental/OnlineServicesEOS/Source/Private/ExternalUIEOS.cpp
--- a/Engine/Plugins/Online/Experimental/OnlineServicesEOS/Source/Private/ExternalUIEOS.cpp
+++ b/Engine/Plugins/Online/Experimental/OnlineServicesEOS/Source/Private/ExternalUIEOS.cpp
@@ -2,11 +2,11 @@
 
 #include "ExternalUIEOS.h"
 
+#include "AuthEOS.h"
+#include "OnlineIdEOS.h"
 #include "OnlineServicesEOS.h"
 #include "OnlineServicesEOSTypes.h"
-//#include "Online/ExternalUIErrors.h"
 #include "Online/OnlineErrorDefinitions.h"
-#include "AuthEOS.h"
 
 #include "eos_ui.h"
 #include "eos_types.h"
@@ -37,21 +37,21 @@ TOnlineAsyncOpHandle<FExternalUIShowFriendsUI> FExternalUIEOS::ShowFriendsUI(FEx
 {
 	TOnlineAsyncOpRef<FExternalUIShowFriendsUI> Op = GetOp<FExternalUIShowFriendsUI>(MoveTemp(Params));
 
-	TOptional<EOS_EpicAccountId> EpicAccountId = EOSAccountIdFromOnlineServiceAccountId(Params.AccountId);
-	if (!EpicAccountId.IsSet())
+	EOS_EpicAccountId LocalUserEasId = GetEpicAccountId(Params.LocalUserId);
+	if (!EOS_EpicAccountId_IsValid(LocalUserEasId))
 	{
-		UE_LOG(LogTemp, Warning, TEXT("[FExternalUIEOS::ShowFriendsUI] EpicAccountId not found for AccountId [%d]"), Params.AccountId.Handle);
+		UE_LOG(LogTemp, Warning, TEXT("[FExternalUIEOS::ShowFriendsUI] LocalUserId=[%s] EpicAccountId not found"), *ToLogString(Params.LocalUserId));
 		Op->SetError(Errors::Unknown()); // TODO
 		return Op->GetHandle();
 	}
 
 	EOS_UI_ShowFriendsOptions ShowFriendsOptions = {};
 	ShowFriendsOptions.ApiVersion = EOS_UI_SHOWFRIENDS_API_LATEST;
-	ShowFriendsOptions.LocalUserId = *EpicAccountId;
+	ShowFriendsOptions.LocalUserId = LocalUserEasId;
 
 	Op->Then([this, ShowFriendsOptions](TOnlineAsyncOp<FExternalUIShowFriendsUI>& InAsyncOp) mutable
 	{
-		return EOS_Async<EOS_UI_ShowFriendsCallbackInfo>(InAsyncOp, EOS_UI_ShowFriends, UIHandle, ShowFriendsOptions);
+		return EOS_Async<EOS_UI_ShowFriendsCallbackInfo>(EOS_UI_ShowFriends, UIHandle, ShowFriendsOptions);
 	})
 	.Then([this](TOnlineAsyncOp<FExternalUIShowFriendsUI>& InAsyncOp, const EOS_UI_ShowFriendsCallbackInfo* Data)
 	{
diff --git a/Engine/Plugins/Online/Experimental/OnlineServicesEOS/Source/Private/FriendsEOS.cpp b/Engine/Plugins/Online/Experimental/OnlineServicesEOS/Source/Private/FriendsEOS.cpp
--- a/Engine/Plugins/Online/Experimental/OnlineServicesEOS/Source/Private/FriendsEOS.cpp
+++ b/Engine/Plugins/Online/Experimental/OnlineServicesEOS/Source/Private/FriendsEOS.cpp
@@ -4,6 +4,7 @@
 #include "Online/OnlineAsyncOp.h"
 #include "OnlineServicesEOS.h"
 #include "OnlineServicesEOSTypes.h"
+#include "OnlineIdEOS.h"
 #include "AuthEOS.h"
 
 #include "eos_friends.h"
@@ -44,9 +45,13 @@ void FFriendsEOS::Initialize()
 	{
 		FFriendsEOS* This = reinterpret_cast<FFriendsEOS*>(Data->ClientData);
 
-		FAccountId LocalUserId = MakeEOSAccountId(Data->LocalUserId);
-		FAccountId FriendUserId = MakeEOSAccountId(Data->TargetUserId);
-		This->OnEOSFriendsUpdate(LocalUserId, FriendUserId, Data->PreviousStatus, Data->CurrentStatus);
+		const FOnlineAccountIdHandle LocalUserId = FindAccountIdChecked(Data->LocalUserId);
+		This->Services.Get<FAuthEOS>()->ResolveAccountId(LocalUserId, Data->TargetUserId)
+		.Next([This, LocalUserId, PreviousStatus = Data->PreviousStatus, CurrentStatus = Data->CurrentStatus](const FOnlineAccountIdHandle& FriendUserId)
+		{
+			UE_LOG(LogTemp, Warning, TEXT("OnEOSPresenceUpdate: [%s] [%s]"), *ToLogString(LocalUserId), *ToLogString(FriendUserId));
+			This->OnEOSFriendsUpdate(LocalUserId, FriendUserId, PreviousStatus, CurrentStatus);
+		});
 	});
 }
 
@@ -71,116 +76,122 @@ TOnlineAsyncOpHandle<FQueryFriends> FFriendsEOS::QueryFriends(FQueryFriends::Par
 
 		EOS_Friends_QueryFriendsOptions QueryFriendsOptions = { };
 		QueryFriendsOptions.ApiVersion = EOS_FRIENDS_QUERYFRIENDS_API_LATEST;
-		TOptional<EOS_EpicAccountId> AccountId = EOSAccountIdFromOnlineServiceAccountId(Params.LocalUserId);
-		check(AccountId); // Must be valid if IsLoggedIn succeeded
-		QueryFriendsOptions.LocalUserId = AccountId.GetValue();
+		EOS_EpicAccountId LocalUserEasId = GetEpicAccountIdChecked(Params.LocalUserId);
+		QueryFriendsOptions.LocalUserId = LocalUserEasId;
 
 		Op->Then([this, QueryFriendsOptions](TOnlineAsyncOp<FQueryFriends>& InAsyncOp) mutable
-			{
-				return EOS_Async<EOS_Friends_QueryFriendsCallbackInfo>(InAsyncOp, EOS_Friends_QueryFriends, FriendsHandle, QueryFriendsOptions);
-			})
-			.Then([this](TOnlineAsyncOp<FQueryFriends>& InAsyncOp, const EOS_Friends_QueryFriendsCallbackInfo* Data) mutable
-			{
-				UE_LOG(LogTemp, Warning, TEXT("QueryFriendsResult: [%s]"), *LexToString(Data->ResultCode));
-
-				if (Data->ResultCode == EOS_EResult::EOS_Success)
-				{
-					bool bAnyChanges = false;
-					EOS_Friends_GetFriendsCountOptions GetFriendsCountOptions = {
-						EOS_FRIENDS_GETFRIENDSCOUNT_API_LATEST,
-						Data->LocalUserId
-					};
-
-					TMap<FAccountId, TSharedRef<FFriend>>& FriendsList = FriendsLists.FindOrAdd(InAsyncOp.GetParams().LocalUserId);
-					TArray<FAccountId> NewFriendIds;
-					TArray<FAccountId> PreviousFriendIds;
-					TArray<FAccountId> UpdatedFriendIds; // TODO:  More granular updates like what changed?
-					FriendsList.GenerateKeyArray(PreviousFriendIds);
-
-					int32_t NumFriends = EOS_Friends_GetFriendsCount(FriendsHandle, &GetFriendsCountOptions);
-					for (int32_t FriendIndex = 0; FriendIndex < NumFriends; ++FriendIndex)
-					{
-						EOS_Friends_GetFriendAtIndexOptions GetFriendAtIndexOptions = {
-							EOS_FRIENDS_GETFRIENDATINDEX_API_LATEST,
-							Data->LocalUserId,
-							FriendIndex
-						};
-
-						EOS_EpicAccountId EOSFriendId = EOS_Friends_GetFriendAtIndex(FriendsHandle, &GetFriendAtIndexOptions);
-						FAccountId FriendId = MakeEOSAccountId(EOSFriendId);
-						//check(FriendId.IsValid()); // TODO:  Some way to check validity
-
-						EOS_Friends_GetStatusOptions GetStatusOptions = {
-							EOS_FRIENDS_GETSTATUS_API_LATEST,
-							Data->LocalUserId,
-							EOSFriendId
-						};
-
-						EOS_EFriendsStatus EOSFriendStatus = EOS_Friends_GetStatus(FriendsHandle, &GetStatusOptions);
-						TOptional<EFriendInviteStatus> InviteStatus = EOSFriendStatusToInviteStatus(EOSFriendStatus);
-
-						PreviousFriendIds.Remove(FriendId);
-						if (TSharedRef<FFriend>* FriendPtr = FriendsList.Find(FriendId))
-						{
-							// Any updates?
-							if (InviteStatus != (*FriendPtr)->InviteStatus)
-							{
-								UE_LOG(LogTemp, Warning, TEXT("QueryFriendsComplete: %s Invite status changed"), *LexToString(EOSFriendId));
-								bAnyChanges = true;
-								(*FriendPtr)->InviteStatus = InviteStatus;
-								UpdatedFriendIds.Emplace(FriendId);
-							}
-						}
-						else
-						{
-							// Add new friend
-							UE_LOG(LogTemp, Warning, TEXT("QueryFriendsComplete: Adding %s"), *LexToString(EOSFriendId));
-							TSharedRef<FFriend> Friend = MakeShared<FFriend>();
-							Friend->UserId = FriendId;
-							Friend->InviteStatus = InviteStatus;
-							FriendsList.Emplace(FriendId, MoveTemp(Friend));
-
-							bAnyChanges = true;
-							NewFriendIds.Emplace(FriendId);
-						}
-					}
+		{
+			return EOS_Async<EOS_Friends_QueryFriendsCallbackInfo>(EOS_Friends_QueryFriends, FriendsHandle, QueryFriendsOptions);
+		})
+		.Then([this](TOnlineAsyncOp<FQueryFriends>& InAsyncOp, const EOS_Friends_QueryFriendsCallbackInfo* Data) mutable
+		{
+			UE_LOG(LogTemp, Warning, TEXT("QueryFriendsResult: [%s]"), *LexToString(Data->ResultCode));
 
-					// TODO:  Delegates
-					bAnyChanges |= NewFriendIds.Num() > 0;
-					bAnyChanges |= PreviousFriendIds.Num() > 0;
-					bAnyChanges |= UpdatedFriendIds.Num() > 0;
+			if (Data->ResultCode != EOS_EResult::EOS_Success)
+			{
+				InAsyncOp.SetError(Errors::Unknown()); // TODO: Error codes
+				return MakeFulfilledPromise<TArray<EOS_EpicAccountId>>().GetFuture();
+			}
 
-					for (FAccountId AccountId : PreviousFriendIds)
-					{
-						UE_LOG(LogTemp, Warning, TEXT("QueryFriendsComplete: %s is no longer a friend, removing"), *LexToString(EOSAccountIdFromOnlineServiceAccountId(AccountId).GetValue()));
-						FriendsList.Remove(AccountId);
-					}
+			EOS_Friends_GetFriendsCountOptions GetFriendsCountOptions = {
+				EOS_FRIENDS_GETFRIENDSCOUNT_API_LATEST,
+				Data->LocalUserId
+			};
 
-					InAsyncOp.SetResult(FQueryFriends::Result());
-					
-					if (bAnyChanges)
+			int32_t NumFriends = EOS_Friends_GetFriendsCount(FriendsHandle, &GetFriendsCountOptions);
+			TArray<EOS_EpicAccountId> CurrentFriends;
+			CurrentFriends.Reserve(NumFriends);
+			for (int32_t FriendIndex = 0; FriendIndex < NumFriends; ++FriendIndex)
+			{
+				EOS_Friends_GetFriendAtIndexOptions GetFriendAtIndexOptions = {
+					EOS_FRIENDS_GETFRIENDATINDEX_API_LATEST,
+					Data->LocalUserId,
+					FriendIndex
+				};
+
+				const EOS_EpicAccountId FriendEasId = EOS_Friends_GetFriendAtIndex(FriendsHandle, &GetFriendAtIndexOptions);
+				check(EOS_EpicAccountId_IsValid(FriendEasId));
+				CurrentFriends.Emplace(FriendEasId);
+			}
+			return MakeFulfilledPromise<TArray<EOS_EpicAccountId>>(CurrentFriends).GetFuture();
+		}).Then(Services.Get<FAuthEOS>()->ResolveEpicIdsFn())
+		.Then([this](TOnlineAsyncOp<FQueryFriends>& InAsyncOp, const TArray<FOnlineAccountIdHandle>& CurrentFriendIds) mutable
+		{
+			const FOnlineAccountIdHandle LocalUserId = InAsyncOp.GetParams().LocalUserId;
+			const EOS_EpicAccountId LocalUserEasId = GetEpicAccountIdChecked(LocalUserId);
+
+			TMap<FOnlineAccountIdHandle, TSharedRef<FFriend>>& FriendsList = FriendsLists.FindOrAdd(LocalUserId);
+			TArray<FOnlineAccountIdHandle> NewFriendIds;
+			TArray<FOnlineAccountIdHandle> PreviousFriendIds;
+			TArray<FOnlineAccountIdHandle> UpdatedFriendIds; // TODO:  More granular updates like what changed?
+			FriendsList.GenerateKeyArray(PreviousFriendIds);
+				
+			for (const FOnlineAccountIdHandle& FriendId : CurrentFriendIds)
+			{
+				const EOS_EpicAccountId FriendEasId = GetEpicAccountIdChecked(FriendId);
+
+				EOS_Friends_GetStatusOptions GetStatusOptions = {
+					EOS_FRIENDS_GETSTATUS_API_LATEST,
+					LocalUserEasId,
+					FriendEasId
+				};
+				const EOS_EFriendsStatus EOSFriendStatus = EOS_Friends_GetStatus(FriendsHandle, &GetStatusOptions);
+				TOptional<EFriendInviteStatus> InviteStatus = EOSFriendStatusToInviteStatus(EOSFriendStatus);
+
+				PreviousFriendIds.Remove(FriendId);
+				if (TSharedRef<FFriend>* FriendPtr = FriendsList.Find(FriendId))
+				{
+					// Any updates?
+					if (InviteStatus != (*FriendPtr)->InviteStatus)
 					{
-						FFriendsListUpdated FriendsListUpdatedParams = {};
-						FriendsListUpdatedParams.LocalUserId = InAsyncOp.GetParams().LocalUserId;
-						FriendsListUpdatedParams.AddedFriends = NewFriendIds;
-						FriendsListUpdatedParams.RemovedFriends = PreviousFriendIds;
-						FriendsListUpdatedParams.UpdatedFriends = UpdatedFriendIds;
-						OnFriendsListUpdatedEvent.Broadcast(FriendsListUpdatedParams);
+						UE_LOG(LogTemp, Warning, TEXT("QueryFriendsComplete: %s Invite status changed"), *LexToString(FriendEasId));
+						(*FriendPtr)->InviteStatus = InviteStatus;
+						UpdatedFriendIds.Emplace(FriendId);
 					}
 				}
 				else
 				{
-					InAsyncOp.SetError(Errors::Unknown()); // TODO: Error codes
+					// Add new friend
+					UE_LOG(LogTemp, Warning, TEXT("QueryFriendsComplete: Adding %s"), *LexToString(FriendEasId));
+					TSharedRef<FFriend> Friend = MakeShared<FFriend>();
+					Friend->UserId = FriendId;
+					Friend->InviteStatus = InviteStatus;
+					FriendsList.Emplace(FriendId, MoveTemp(Friend));
+					NewFriendIds.Emplace(FriendId);
 				}
-			})
-			.Enqueue(GetSerialQueue());
+			}
+
+			// TODO:  Delegates
+
+			for (const FOnlineAccountIdHandle& AccountId : PreviousFriendIds)
+			{
+				UE_LOG(LogTemp, Warning, TEXT("QueryFriendsComplete: %s is no longer a friend, removing"), *ToLogString(AccountId));
+				FriendsList.Remove(AccountId);
+			}
+
+			InAsyncOp.SetResult(FQueryFriends::Result());
+					
+			const bool bAnyChanges = NewFriendIds.Num() > 0
+				|| PreviousFriendIds.Num() > 0
+				|| UpdatedFriendIds.Num() > 0;
+			if (bAnyChanges)
+			{
+				FFriendsListUpdated FriendsListUpdatedParams = {};
+				FriendsListUpdatedParams.LocalUserId = InAsyncOp.GetParams().LocalUserId;
+				FriendsListUpdatedParams.AddedFriends = NewFriendIds;
+				FriendsListUpdatedParams.RemovedFriends = PreviousFriendIds;
+				FriendsListUpdatedParams.UpdatedFriends = UpdatedFriendIds;
+				OnFriendsListUpdatedEvent.Broadcast(FriendsListUpdatedParams);
+			}
+		})
+		.Enqueue(GetSerialQueue());
 	}
 	return Op->GetHandle();
 }
 
 TOnlineResult<FGetFriends> FFriendsEOS::GetFriends(FGetFriends::Params&& Params)
 {
-	if (TMap<FAccountId, TSharedRef<FFriend>>* FriendsList = FriendsLists.Find(Params.LocalUserId))
+	if (TMap<FOnlineAccountIdHandle, TSharedRef<FFriend>>* FriendsList = FriendsLists.Find(Params.LocalUserId))
 	{
 		FGetFriends::Result Result;
 		FriendsList->GenerateValueArray(Result.Friends);
@@ -205,27 +216,28 @@ TOnlineAsyncOpHandle<FAddFriend> FFriendsEOS::AddFriend(FAddFriend::Params&& InP
 		}
 
 		// Target user valid?
-		TOptional<EOS_EpicAccountId> TargetId = EOSAccountIdFromOnlineServiceAccountId(Params.FriendId);
-		if (!TargetId)
+		EOS_EpicAccountId TargetId = GetEpicAccountId(Params.FriendId);
+		if (!EOS_EpicAccountId_IsValid(TargetId))
 		{
 			// TODO: Error codes
 			Op->SetError(Errors::UnknownError());
 			return Op->GetHandle();
 		}
+
 		Op->Then([this](TOnlineAsyncOp<FAddFriend>& InAsyncOp) mutable
-			{
-				const FAddFriend::Params& Params = InAsyncOp.GetParams();
-				EOS_Friends_SendInviteOptions SendInviteOptions = { };
-				SendInviteOptions.ApiVersion = EOS_FRIENDS_SENDINVITE_API_LATEST;
-				SendInviteOptions.LocalUserId = EOSAccountIdFromOnlineServiceAccountId(Params.LocalUserId).GetValue();
-				SendInviteOptions.TargetUserId = EOSAccountIdFromOnlineServiceAccountId(Params.FriendId).GetValue();
-				return EOS_Async<EOS_Friends_SendInviteCallbackInfo>(InAsyncOp, EOS_Friends_SendInvite, FriendsHandle, SendInviteOptions);
-			})
-			.Then([this](TOnlineAsyncOp<FAddFriend>& InAsyncOp, const EOS_Friends_SendInviteCallbackInfo* Data) mutable
-			{
-				UE_LOG(LogTemp, Warning, TEXT("SendInviteResult: [%s]"), *LexToString(Data->ResultCode));
-				// TODO:  Handle response
-			}).Enqueue();
+		{
+			const FAddFriend::Params& Params = InAsyncOp.GetParams();
+			EOS_Friends_SendInviteOptions SendInviteOptions = { };
+			SendInviteOptions.ApiVersion = EOS_FRIENDS_SENDINVITE_API_LATEST;
+			SendInviteOptions.LocalUserId = EOSAccountIdFromOnlineServiceAccountId(Params.LocalUserId).GetValue();
+			SendInviteOptions.TargetUserId = EOSAccountIdFromOnlineServiceAccountId(Params.FriendId).GetValue();
+			return EOS_Async<EOS_Friends_SendInviteCallbackInfo>(EOS_Friends_SendInvite, FriendsHandle, SendInviteOptions);
+		})
+		.Then([this](TOnlineAsyncOp<FAddFriend>& InAsyncOp, const EOS_Friends_SendInviteCallbackInfo* Data) mutable
+		{
+			UE_LOG(LogTemp, Warning, TEXT("SendInviteResult: [%s]"), *LexToString(Data->ResultCode));
+			// TODO:  Handle response
+		}).Enqueue();
 #else
 		Op->SetError(Errors::Unknown());
 #endif
@@ -233,20 +245,16 @@ TOnlineAsyncOpHandle<FAddFriend> FFriendsEOS::AddFriend(FAddFriend::Params&& InP
 	return Op->GetHandle();
 }
 
-void FFriendsEOS::OnEOSFriendsUpdate(FAccountId LocalUserId, FAccountId FriendUserId, EOS_EFriendsStatus PreviousStatus, EOS_EFriendsStatus CurrentStatus)
+void FFriendsEOS::OnEOSFriendsUpdate(FOnlineAccountIdHandle LocalUserId, FOnlineAccountIdHandle FriendUserId, EOS_EFriendsStatus PreviousStatus, EOS_EFriendsStatus CurrentStatus)
 {
 	// TODO
-	UE_LOG(LogTemp, Warning, TEXT("OnEOSFriendsUpdate: [%s] [%s] %d %d"), 
-		*LexToString(EOSAccountIdFromOnlineServiceAccountId(LocalUserId).GetValue()),
-		*LexToString(EOSAccountIdFromOnlineServiceAccountId(FriendUserId).GetValue()),
-		(int)PreviousStatus,
-		(int)CurrentStatus);
+	UE_LOG(LogTemp, Warning, TEXT("OnEOSFriendsUpdate: LocalUserId=[%s] FriendUserId=[%s] PreviousStatus=[%s] CurrentStatus=[%s]"), *ToLogString(LocalUserId), *ToLogString(FriendUserId), *LexToString(PreviousStatus), *LexToString(CurrentStatus));
 
-	TMap<FAccountId, TSharedRef<FFriend>>& FriendsList = FriendsLists.FindOrAdd(LocalUserId);
+	TMap<FOnlineAccountIdHandle, TSharedRef<FFriend>>& FriendsList = FriendsLists.FindOrAdd(LocalUserId);
 	bool bAnyChanges = false;
-	TArray<FAccountId, TInlineAllocator<1>> AddedFriends;
-	TArray<FAccountId, TInlineAllocator<1>> RemovedFriends;
-	TArray<FAccountId, TInlineAllocator<1>> UpdatedFriends;
+	TArray<FOnlineAccountIdHandle, TInlineAllocator<1>> AddedFriends;
+	TArray<FOnlineAccountIdHandle, TInlineAllocator<1>> RemovedFriends;
+	TArray<FOnlineAccountIdHandle, TInlineAllocator<1>> UpdatedFriends;
 	if (CurrentStatus == EOS_EFriendsStatus::EOS_FS_NotFriends)
 	{
 		bAnyChanges = FriendsList.Remove(FriendUserId) > 0;
diff --git a/Engine/Plugins/Online/Experimental/OnlineServicesEOS/Source/Private/OnlineIdEOS.cpp b/Engine/Plugins/Online/Experimental/OnlineServicesEOS/Source/Private/OnlineIdEOS.cpp
new file mode 100644
--- /dev/null
+++ b/Engine/Plugins/Online/Experimental/OnlineServicesEOS/Source/Private/OnlineIdEOS.cpp
@@ -0,0 +1,178 @@
+// Copyright Epic Games, Inc. All Rights Reserved.
+
+#include "OnlineIdEOS.h"
+
+#include "OnlineServicesEOSTypes.h"
+#include "Online/OnlineAsyncOp.h"
+
+#include "eos_connect_types.h"
+#include "eos_connect.h"
+
+namespace UE::Online {
+
+FOnlineAccountIdRegistryEOS& FOnlineAccountIdRegistryEOS::Get()
+{
+	static FOnlineAccountIdRegistryEOS Instance;
+	return Instance;
+}
+
+FOnlineAccountIdHandle FOnlineAccountIdRegistryEOS::Find(const EOS_EpicAccountId EpicAccountId) const
+{
+	FOnlineAccountIdHandle Handle;
+	Handle.Type = EOnlineServices::Epic;
+	if (EOS_EpicAccountId_IsValid(EpicAccountId) == EOS_TRUE)
+	{
+		const FReadScopeLock ReadLock(Lock);
+		if (const uint32* Found = EpicAccountIdToHandle.Find(EpicAccountId))
+		{
+			Handle.Handle = *Found;
+		}
+	}
+	return Handle;
+}
+
+FOnlineAccountIdHandle FOnlineAccountIdRegistryEOS::Find(const EOS_ProductUserId ProductUserId) const
+{
+	FOnlineAccountIdHandle Handle;
+	Handle.Type = EOnlineServices::Epic;
+	if (EOS_ProductUserId_IsValid(ProductUserId) == EOS_TRUE)
+	{
+		const FReadScopeLock ReadLock(Lock);
+		if (const uint32* Found = ProductUserIdToHandle.Find(ProductUserId))
+		{
+			Handle.Handle = *Found;
+		}
+	}
+	return Handle;
+}
+
+FOnlineAccountIdHandle FOnlineAccountIdRegistryEOS::Create(const EOS_EpicAccountId EpicAccountId, const EOS_ProductUserId ProductUserId)
+{
+	FOnlineAccountIdHandle Result;
+	Result.Type = EOnlineServices::Epic;
+	const bool bEpicAccountIdValid = EOS_EpicAccountId_IsValid(EpicAccountId) == EOS_TRUE;
+	const bool bProductUserIdValid = EOS_ProductUserId_IsValid(ProductUserId) == EOS_TRUE;
+	check(bEpicAccountIdValid || bProductUserIdValid);
+
+	if (bEpicAccountIdValid || bProductUserIdValid)
+	{
+		auto FindExisting = [&]() -> uint32
+		{
+			uint32 Result = 0;
+			if (const uint32* FoundEas = bEpicAccountIdValid ? EpicAccountIdToHandle.Find(EpicAccountId) : nullptr)
+			{
+				Result = *FoundEas;
+			}
+			else if (const uint32* FoundProd = bProductUserIdValid ? ProductUserIdToHandle.Find(ProductUserId) : nullptr)
+			{
+				Result = *FoundProd;
+			}
+
+			if (Result != 0)
+			{
+				check(EpicAccountId == AccountIdData[Result - 1].EpicAccountId);
+				check(ProductUserId == AccountIdData[Result - 1].ProductUserId);
+			}
+
+			return Result;
+		};
+
+		// First take read lock and look for existing elements
+		{
+			const FReadScopeLock ReadLock(Lock);
+			Result.Handle = FindExisting();
+		}
+
+		if(!Result.IsValid())
+		{
+			// Next take write lock, check again for existing elements (in case another thread added it)
+			const FWriteScopeLock WriteLock(Lock);
+			Result.Handle = FindExisting();
+			if (!Result.IsValid())
+			{
+				// No existing element, so add a new one
+				FOnlineAccountIdDataEOS& NewAccountIdData = AccountIdData.Emplace_GetRef();
+				NewAccountIdData.EpicAccountId = EpicAccountId;
+				NewAccountIdData.ProductUserId = ProductUserId;
+
+				Result.Handle = AccountIdData.Num();
+
+				if (bEpicAccountIdValid)
+				{
+					EpicAccountIdToHandle.Emplace(EpicAccountId, CopyTemp(Result.Handle));
+				}
+				if (bProductUserIdValid)
+				{
+					ProductUserIdToHandle.Emplace(ProductUserId, CopyTemp(Result.Handle));
+				}
+			}
+		}
+	}
+	
+	return Result;
+}
+
+FOnlineAccountIdDataEOS FOnlineAccountIdRegistryEOS::GetIdData(const FOnlineAccountIdHandle& Handle) const
+{
+	if (Handle.Type == EOnlineServices::Epic && Handle.IsValid())
+	{
+		FReadScopeLock ScopeLock(Lock);
+		return AccountIdData[Handle.Handle - 1];
+	}
+	return FOnlineAccountIdDataEOS();
+}
+
+EOS_EpicAccountId GetEpicAccountId(const FOnlineAccountIdHandle& Handle)
+{
+	return FOnlineAccountIdRegistryEOS::Get().GetIdData(Handle).EpicAccountId;
+}
+
+EOS_EpicAccountId GetEpicAccountIdChecked(const FOnlineAccountIdHandle& Handle)
+{
+	EOS_EpicAccountId EpicAccountId = GetEpicAccountId(Handle);
+	check(EOS_EpicAccountId_IsValid(EpicAccountId) == EOS_TRUE);
+	return EpicAccountId;
+}
+
+EOS_ProductUserId GetProductUserId(const FOnlineAccountIdHandle& Handle)
+{
+	return FOnlineAccountIdRegistryEOS::Get().GetIdData(Handle).ProductUserId;
+}
+
+EOS_ProductUserId GetProductUserIdChecked(const FOnlineAccountIdHandle& Handle)
+{
+	EOS_ProductUserId ProductUserId = GetProductUserId(Handle);
+	check(EOS_ProductUserId_IsValid(ProductUserId) == EOS_TRUE);
+	return ProductUserId;
+}
+
+FOnlineAccountIdHandle FindAccountId(const EOS_EpicAccountId EpicAccountId)
+{
+	return FOnlineAccountIdRegistryEOS::Get().Find(EpicAccountId);
+}
+
+FOnlineAccountIdHandle FindAccountId(const EOS_ProductUserId ProductUserId)
+{
+	return FOnlineAccountIdRegistryEOS::Get().Find(ProductUserId);
+}
+
+FOnlineAccountIdHandle FindAccountIdChecked(const EOS_EpicAccountId EpicAccountId)
+{
+	FOnlineAccountIdHandle Result = FindAccountId(EpicAccountId);
+	check(Result.IsValid());
+	return Result;
+}
+
+FOnlineAccountIdHandle FindAccountIdChecked(const EOS_ProductUserId ProductUserId)
+{
+	FOnlineAccountIdHandle Result = FindAccountId(ProductUserId);
+	check(Result.IsValid());
+	return Result;
+}
+
+FOnlineAccountIdHandle CreateAccountId(const EOS_EpicAccountId EpicAccountId, const EOS_ProductUserId ProductUserId)
+{
+	return FOnlineAccountIdRegistryEOS::Get().Create(EpicAccountId, ProductUserId);
+}
+
+/* UE::Online */ }
diff --git a/Engine/Plugins/Online/Experimental/OnlineServicesEOS/Source/Private/OnlineIdEOS.h b/Engine/Plugins/Online/Experimental/OnlineServicesEOS/Source/Private/OnlineIdEOS.h
new file mode 100644
--- /dev/null
+++ b/Engine/Plugins/Online/Experimental/OnlineServicesEOS/Source/Private/OnlineIdEOS.h
@@ -0,0 +1,67 @@
+// Copyright Epic Games, Inc. All Rights Reserved.
+
+#pragma once
+
+#include "Async/Future.h"
+#include "Delegates/Delegate.h"
+#include "Online/OnlineId.h"
+#include "Online/OnlineIdCommon.h"
+
+#if defined(EOS_PLATFORM_BASE_FILE_NAME)
+#include EOS_PLATFORM_BASE_FILE_NAME
+#endif
+#include "eos_common.h"
+#include "eos_connect_types.h"
+
+namespace UE::Online {
+
+class FOnlineAccountIdDataEOS
+{
+public:
+	EOS_EpicAccountId EpicAccountId = nullptr;
+	EOS_ProductUserId ProductUserId = nullptr;
+};
+
+/**
+ * Account id registry specifically for EOS id's which are segmented.
+ */
+class FOnlineAccountIdRegistryEOS
+{
+public:
+	static FOnlineAccountIdRegistryEOS& Get();
+
+	FOnlineAccountIdHandle Find(const EOS_EpicAccountId EpicAccountId) const;
+	FOnlineAccountIdHandle Find(const EOS_ProductUserId ProductUserId) const;
+	FOnlineAccountIdHandle Create(const EOS_EpicAccountId EpicAccountId, const EOS_ProductUserId ProductUserId);
+
+	// Return copies as it is not thread safe to return pointers/references to array elements, in case the array is grown+relocated on another thread.
+	FOnlineAccountIdDataEOS GetIdData(const FOnlineAccountIdHandle& Handle) const;
+
+private:
+	mutable FRWLock Lock;
+
+	TArray<FOnlineAccountIdDataEOS> AccountIdData; // Actual container for the info, indexed by the handle
+
+	TMap<EOS_EpicAccountId, uint32> EpicAccountIdToHandle; // Map of EOS_EpicAccountId to the associated handle.
+	TMap<EOS_ProductUserId, uint32> ProductUserIdToHandle; // Map of EOS_ProductUserId to the associated handle.
+};
+
+EOS_EpicAccountId GetEpicAccountId(const FOnlineAccountIdHandle& Handle);
+EOS_EpicAccountId GetEpicAccountIdChecked(const FOnlineAccountIdHandle& Handle);
+EOS_ProductUserId GetProductUserId(const FOnlineAccountIdHandle& Handle);
+EOS_ProductUserId GetProductUserIdChecked(const FOnlineAccountIdHandle& Handle);
+
+FOnlineAccountIdHandle FindAccountId(const EOS_EpicAccountId EpicAccountId);
+FOnlineAccountIdHandle FindAccountId(const EOS_ProductUserId EpicAccountId);
+FOnlineAccountIdHandle FindAccountIdChecked(const EOS_EpicAccountId EpicAccountId);
+FOnlineAccountIdHandle FindAccountIdChecked(const EOS_ProductUserId EpicAccountId);
+
+FOnlineAccountIdHandle CreateAccountId(const EOS_EpicAccountId EpicAccountId, const EOS_ProductUserId ProductUserId);
+
+template<EOnlineIdType IdType>
+inline bool ValidateOnlineId(const TOnlineIdHandle<IdType>& Handle)
+{
+	return Handle.Type == EOnlineServices::Epic && Handle.IsValid();
+}
+
+} /* namespace UE::Online */
diff --git a/Engine/Plugins/Online/Experimental/OnlineServicesEOS/Source/Private/OnlineServicesEOS.cpp b/Engine/Plugins/Online/Experimental/OnlineServicesEOS/Source/Private/OnlineServicesEOS.cpp
--- a/Engine/Plugins/Online/Experimental/OnlineServicesEOS/Source/Private/OnlineServicesEOS.cpp
+++ b/Engine/Plugins/Online/Experimental/OnlineServicesEOS/Source/Private/OnlineServicesEOS.cpp
@@ -3,9 +3,12 @@
 #include "OnlineServicesEOS.h"
 #include "OnlineServicesEOSTypes.h"
 
+#include "EOSShared.h"
 #include "IEOSSDKManager.h"
+
 #include "AuthEOS.h"
 #include "FriendsEOS.h"
+#include "OnlineIdEOS.h"
 #include "PresenceEOS.h"
 #include "ExternalUIEOS.h"
 
@@ -85,14 +88,19 @@ void FOnlineServicesEOS::Initialize()
 	FOnlineServicesCommon::Initialize();
 }
 
-FAccountId FOnlineServicesEOS::CreateAccountId(FString&& InAccountIdString)
+FString FOnlineServicesEOS::ToDebugString(const FOnlineAccountIdHandle& Handle)
 {
-	EOS_EpicAccountId EpicAccountId = EOS_EpicAccountId_FromString(TCHAR_TO_UTF8(*InAccountIdString));
-	if (EOS_EpicAccountId_IsValid(EpicAccountId) == EOS_TRUE)
+	FString Result;
+	if (Handle.IsValid())
+	{
+		const FOnlineAccountIdDataEOS& AccountIdData = FOnlineAccountIdRegistryEOS::Get().GetIdData(Handle);
+		Result = FString::Printf(TEXT("EAS=[%s] EOS=[%s]"), *LexToString(AccountIdData.EpicAccountId), *LexToString(AccountIdData.ProductUserId));
+	}
+	else
 	{
-		return MakeEOSAccountId(EpicAccountId);
+		Result = TEXT("Invalid");
 	}
-	return FAccountId();
+	return Result;
 }
 
 EOS_HPlatform FOnlineServicesEOS::GetEOSPlatformHandle() const
diff --git a/Engine/Plugins/Online/Experimental/OnlineServicesEOS/Source/Private/PresenceEOS.cpp b/Engine/Plugins/Online/Experimental/OnlineServicesEOS/Source/Private/PresenceEOS.cpp
--- a/Engine/Plugins/Online/Experimental/OnlineServicesEOS/Source/Private/PresenceEOS.cpp
+++ b/Engine/Plugins/Online/Experimental/OnlineServicesEOS/Source/Private/PresenceEOS.cpp
@@ -2,6 +2,7 @@
 
 #include "PresenceEOS.h"
 #include "Online/OnlineAsyncOp.h"
+#include "OnlineIdEOS.h"
 #include "OnlineServicesEOS.h"
 #include "OnlineServicesEOSTypes.h"
 #include "AuthEOS.h"
@@ -76,13 +77,14 @@ void FPresenceEOS::Initialize()
 	NotifyPresenceChangedNotificationId = EOS_Presence_AddNotifyOnPresenceChanged(PresenceHandle, &Options, this, [](const EOS_Presence_PresenceChangedCallbackInfo* Data)
 	{
 		FPresenceEOS* This = reinterpret_cast<FPresenceEOS*>(Data->ClientData);
-		FAccountId LocalUserId = MakeEOSAccountId(Data->LocalUserId);
-		FAccountId PresenceUserId = MakeEOSAccountId(Data->PresenceUserId);
+		const FOnlineAccountIdHandle LocalUserId = FindAccountIdChecked(Data->LocalUserId);
 
-		UE_LOG(LogTemp, Warning, TEXT("OnEOSPresenceUpdate: [%s] [%s]"),
-			*LexToString(Data->LocalUserId),
-			*LexToString(Data->PresenceUserId));
-		This->UpdateUserPresence(LocalUserId, PresenceUserId);
+		This->Services.Get<FAuthEOS>()->ResolveAccountId(LocalUserId, Data->PresenceUserId)
+		.Next([This, LocalUserId](const FOnlineAccountIdHandle& PresenceUserId)
+		{
+			UE_LOG(LogTemp, Warning, TEXT("OnEOSPresenceUpdate: [%s] [%s]"), *ToLogString(LocalUserId), *ToLogString(PresenceUserId));
+			This->UpdateUserPresence(LocalUserId, PresenceUserId);
+		});
 	});
 }
 
@@ -104,8 +106,8 @@ TOnlineAsyncOpHandle<FQueryPresence> FPresenceEOS::QueryPresence(FQueryPresence:
 			Op->SetError(Errors::Unknown());
 			return Op->GetHandle();
 		}
-		TOptional<EOS_EpicAccountId> TargetUserId = EOSAccountIdFromOnlineServiceAccountId(Params.TargetUserId);
-		if (!TargetUserId.IsSet())
+		EOS_EpicAccountId TargetUserEasId = GetEpicAccountId(Params.TargetUserId);
+		if (!EOS_EpicAccountId_IsValid(TargetUserEasId))
 		{
 			// TODO: Error codes
 			Op->SetError(Errors::Unknown());
@@ -125,10 +127,10 @@ TOnlineAsyncOpHandle<FQueryPresence> FPresenceEOS::QueryPresence(FQueryPresence:
 				const FQueryPresence::Params& Params = InAsyncOp.GetParams();
 				EOS_Presence_QueryPresenceOptions QueryPresenceOptions = { };
 				QueryPresenceOptions.ApiVersion = EOS_PRESENCE_QUERYPRESENCE_API_LATEST;
-				QueryPresenceOptions.LocalUserId = EOSAccountIdFromOnlineServiceAccountId(Params.LocalUserId).GetValue();
-				QueryPresenceOptions.TargetUserId = EOSAccountIdFromOnlineServiceAccountId(Params.TargetUserId).GetValue();
+				QueryPresenceOptions.LocalUserId = GetEpicAccountIdChecked(Params.LocalUserId);
+				QueryPresenceOptions.TargetUserId = GetEpicAccountIdChecked(Params.TargetUserId);
 
-				return EOS_Async<EOS_Presence_QueryPresenceCallbackInfo>(InAsyncOp, EOS_Presence_QueryPresence, PresenceHandle, QueryPresenceOptions);
+				return EOS_Async<EOS_Presence_QueryPresenceCallbackInfo>(EOS_Presence_QueryPresence, PresenceHandle, QueryPresenceOptions);
 			})
 			.Then([this](TOnlineAsyncOp<FQueryPresence>& InAsyncOp, const EOS_Presence_QueryPresenceCallbackInfo* Data) mutable
 			{
@@ -153,7 +155,7 @@ TOnlineAsyncOpHandle<FQueryPresence> FPresenceEOS::QueryPresence(FQueryPresence:
 
 TOnlineResult<FGetPresence> FPresenceEOS::GetPresence(FGetPresence::Params&& Params)
 {
-	if (TMap<FAccountId, TSharedRef<FUserPresence>>* PresenceList = PresenceLists.Find(Params.LocalUserId))
+	if (TMap<FOnlineAccountIdHandle, TSharedRef<FUserPresence>>* PresenceList = PresenceLists.Find(Params.LocalUserId))
 	{
 		TSharedRef<FUserPresence>* PresencePtr = PresenceList->Find(Params.TargetUserId);
 		if (PresencePtr)
@@ -192,7 +194,7 @@ TOnlineAsyncOpHandle<FUpdatePresence> FPresenceEOS::UpdatePresence(FUpdatePresen
 				EOS_HPresenceModification ChangeHandle = nullptr;
 				EOS_Presence_CreatePresenceModificationOptions Options = { };
 				Options.ApiVersion = EOS_PRESENCE_CREATEPRESENCEMODIFICATION_API_LATEST;
-				Options.LocalUserId = EOSAccountIdFromOnlineServiceAccountId(Params.LocalUserId).GetValue();
+				Options.LocalUserId = GetEpicAccountIdChecked(Params.LocalUserId);
 				EOS_EResult CreatePresenceModificationResult = EOS_Presence_CreatePresenceModification(PresenceHandle, &Options, &ChangeHandle);
 				if (CreatePresenceModificationResult == EOS_EResult::EOS_Success)
 				{
@@ -311,9 +313,9 @@ TOnlineAsyncOpHandle<FUpdatePresence> FPresenceEOS::UpdatePresence(FUpdatePresen
 
 					EOS_Presence_SetPresenceOptions SetPresenceOptions = { };
 					SetPresenceOptions.ApiVersion = EOS_PRESENCE_SETPRESENCE_API_LATEST;
-					SetPresenceOptions.LocalUserId = EOSAccountIdFromOnlineServiceAccountId(Params.LocalUserId).GetValue();
+					SetPresenceOptions.LocalUserId = GetEpicAccountIdChecked(Params.LocalUserId);
 					SetPresenceOptions.PresenceModificationHandle = ChangeHandle;
-					TFuture<const EOS_Presence_SetPresenceCallbackInfo*> EOSAsyncFuture = EOS_Async<EOS_Presence_SetPresenceCallbackInfo>(InAsyncOp, EOS_Presence_SetPresence, PresenceHandle, SetPresenceOptions);
+					TFuture<const EOS_Presence_SetPresenceCallbackInfo*> EOSAsyncFuture = EOS_Async<EOS_Presence_SetPresenceCallbackInfo>(EOS_Presence_SetPresence, PresenceHandle, SetPresenceOptions);
 					EOS_PresenceModification_Release(ChangeHandle);
 					return EOSAsyncFuture;
 				}
@@ -363,9 +365,9 @@ TOnlineAsyncOpHandle<FUpdatePresence> FPresenceEOS::UpdatePresence(FUpdatePresen
 }
 
 /** Get a user's presence, creating entries if missing */
-TSharedRef<FUserPresence> FPresenceEOS::FindOrCreatePresence(FAccountId LocalUserId, FAccountId PresenceUserId)
+TSharedRef<FUserPresence> FPresenceEOS::FindOrCreatePresence(FOnlineAccountIdHandle LocalUserId, FOnlineAccountIdHandle PresenceUserId)
 {
-	TMap<FAccountId, TSharedRef<FUserPresence>>& LocalUserPresenceList = PresenceLists.FindOrAdd(LocalUserId);
+	TMap<FOnlineAccountIdHandle, TSharedRef<FUserPresence>>& LocalUserPresenceList = PresenceLists.FindOrAdd(LocalUserId);
 	if (const TSharedRef<FUserPresence>* const ExistingPresence = LocalUserPresenceList.Find(PresenceUserId))
 	{
 		return *ExistingPresence;
@@ -377,7 +379,7 @@ TSharedRef<FUserPresence> FPresenceEOS::FindOrCreatePresence(FAccountId LocalUse
 	return UserPresence;
 }
 
-void FPresenceEOS::UpdateUserPresence(FAccountId LocalUserId, FAccountId PresenceUserId)
+void FPresenceEOS::UpdateUserPresence(FOnlineAccountIdHandle LocalUserId, FOnlineAccountIdHandle PresenceUserId)
 {
 	bool bPresenceHasChanged = false;
 	TSharedRef<FUserPresence> UserPresence = FindOrCreatePresence(LocalUserId, PresenceUserId);
@@ -387,8 +389,8 @@ void FPresenceEOS::UpdateUserPresence(FAccountId LocalUserId, FAccountId Presenc
 	EOS_Presence_Info* PresenceInfo = nullptr;
 	EOS_Presence_CopyPresenceOptions Options = { };
 	Options.ApiVersion = EOS_PRESENCE_COPYPRESENCE_API_LATEST;
-	Options.LocalUserId = EOSAccountIdFromOnlineServiceAccountId(LocalUserId).GetValue();
-	Options.TargetUserId = EOSAccountIdFromOnlineServiceAccountId(PresenceUserId).GetValue();
+	Options.LocalUserId = GetEpicAccountIdChecked(LocalUserId);
+	Options.TargetUserId = GetEpicAccountIdChecked(PresenceUserId);
 	EOS_EResult CopyPresenceResult = EOS_Presence_CopyPresence(PresenceHandle, &Options, &PresenceInfo);
 	if (CopyPresenceResult == EOS_EResult::EOS_Success)
 	{
diff --git a/Engine/Plugins/Online/Experimental/OnlineServicesEOS/Source/Public/AuthEOS.h b/Engine/Plugins/Online/Experimental/OnlineServicesEOS/Source/Public/AuthEOS.h
--- a/Engine/Plugins/Online/Experimental/OnlineServicesEOS/Source/Public/AuthEOS.h
+++ b/Engine/Plugins/Online/Experimental/OnlineServicesEOS/Source/Public/AuthEOS.h
@@ -29,11 +29,21 @@ public:
 	virtual TOnlineResult<FAuthGetAccountByLocalUserNum> GetAccountByLocalUserNum(FAuthGetAccountByLocalUserNum::Params&& Params) override;
 	virtual TOnlineResult<FAuthGetAccountByAccountId> GetAccountByAccountId(FAuthGetAccountByAccountId::Params&& Params) override;
 
-	bool IsLoggedIn(const FAccountId& AccountId) const;
+	bool IsLoggedIn(const FOnlineAccountIdHandle& AccountId) const;
+
+	TFuture<FOnlineAccountIdHandle> ResolveAccountId(const FOnlineAccountIdHandle& LocalUserId, const EOS_EpicAccountId EpicAccountId);
+	TFuture<FOnlineAccountIdHandle> ResolveAccountId(const FOnlineAccountIdHandle& LocalUserId, const EOS_ProductUserId ProductUserId);
+	TFuture<TArray<FOnlineAccountIdHandle>> ResolveAccountIds(const FOnlineAccountIdHandle& LocalUserId, const TArray<EOS_EpicAccountId>& EpicAccountIds);
+	TFuture<TArray<FOnlineAccountIdHandle>> ResolveAccountIds(const FOnlineAccountIdHandle& LocalUserId, const TArray<EOS_ProductUserId>& ProductUserIds);
+
+	TFunction<TFuture<FOnlineAccountIdHandle>(FOnlineAsyncOp& InAsyncOp, const EOS_EpicAccountId& EpicAccountId)> ResolveEpicIdFn();
+	TFunction<TFuture<FOnlineAccountIdHandle>(FOnlineAsyncOp& InAsyncOp, const EOS_ProductUserId& ProductUserId)> ResolveProductIdFn();
+	TFunction<TFuture<TArray<FOnlineAccountIdHandle>>(FOnlineAsyncOp& InAsyncOp, const TArray<EOS_EpicAccountId>& EpicAccountIds)> ResolveEpicIdsFn();
+	TFunction<TFuture<TArray<FOnlineAccountIdHandle>>(FOnlineAsyncOp& InAsyncOp, const TArray<EOS_ProductUserId>& ProductUserIds)> ResolveProductIdsFn();
 
 protected:
-	void OnEOSLoginStatusChanged(FAccountId LocalUserId, ELoginStatus PreviousStatus, ELoginStatus CurrentStatus);
-	TResult<FAccountId, FOnlineError> GetAccountIdByLocalUserNum(int32 LocalUserNum) const;
+	void OnEOSLoginStatusChanged(FOnlineAccountIdHandle LocalUserId, ELoginStatus PreviousStatus, ELoginStatus CurrentStatus);
+	TResult<FOnlineAccountIdHandle, FOnlineError> GetAccountIdByLocalUserNum(int32 LocalUserNum) const;
 
 	void ProcessSuccessfulLogin(TOnlineAsyncOp<FAuthLogin>& InAsyncOp);
 
@@ -43,7 +53,7 @@ protected:
 
 	};
 
-	TMap<FAccountId, TSharedRef<FAccountInfoEOS>> AccountInfos;
+	TMap<FOnlineAccountIdHandle, TSharedRef<FAccountInfoEOS>> AccountInfos;
 
 	EOS_HAuth AuthHandle;
 	EOS_HConnect ConnectHandle;
diff --git a/Engine/Plugins/Online/Experimental/OnlineServicesEOS/Source/Public/FriendsEOS.h b/Engine/Plugins/Online/Experimental/OnlineServicesEOS/Source/Public/FriendsEOS.h
--- a/Engine/Plugins/Online/Experimental/OnlineServicesEOS/Source/Public/FriendsEOS.h
+++ b/Engine/Plugins/Online/Experimental/OnlineServicesEOS/Source/Public/FriendsEOS.h
@@ -27,9 +27,9 @@ public:
 	virtual TOnlineAsyncOpHandle<FAddFriend> AddFriend(FAddFriend::Params&& Params) override;
 
 protected:
-	void OnEOSFriendsUpdate(FAccountId LocalUserId, FAccountId FriendUserId, EOS_EFriendsStatus PreviousStatus, EOS_EFriendsStatus CurrentStatus);
+	void OnEOSFriendsUpdate(FOnlineAccountIdHandle LocalUserId, FOnlineAccountIdHandle FriendUserId, EOS_EFriendsStatus PreviousStatus, EOS_EFriendsStatus CurrentStatus);
 protected:
-	TMap<FAccountId, TMap<FAccountId, TSharedRef<FFriend>>> FriendsLists;
+	TMap<FOnlineAccountIdHandle, TMap<FOnlineAccountIdHandle, TSharedRef<FFriend>>> FriendsLists;
 	EOS_HFriends FriendsHandle = nullptr;
 
 	EOS_NotificationId NotifyFriendsUpdateNotificationId = 0;
diff --git a/Engine/Plugins/Online/Experimental/OnlineServicesEOS/Source/Public/OnlineServicesEOS.h b/Engine/Plugins/Online/Experimental/OnlineServicesEOS/Source/Public/OnlineServicesEOS.h
--- a/Engine/Plugins/Online/Experimental/OnlineServicesEOS/Source/Public/OnlineServicesEOS.h
+++ b/Engine/Plugins/Online/Experimental/OnlineServicesEOS/Source/Public/OnlineServicesEOS.h
@@ -22,10 +22,11 @@ public:
 	virtual void RegisterComponents() override;
 	virtual void Initialize() override;
 
-	virtual FAccountId CreateAccountId(FString&& InAccountIdString) override;
 	static EOnlineServices GetServicesProvider() { return EOnlineServices::Epic; }
 
 	EOS_HPlatform GetEOSPlatformHandle() const;
+
+	virtual FString ToDebugString(const FOnlineAccountIdHandle& Handle) override;
 protected:
 	IEOSPlatformHandlePtr EOSPlatformHandle;
 };
diff --git a/Engine/Plugins/Online/Experimental/OnlineServicesEOS/Source/Public/OnlineServicesEOSTypes.h b/Engine/Plugins/Online/Experimental/OnlineServicesEOS/Source/Public/OnlineServicesEOSTypes.h
--- a/Engine/Plugins/Online/Experimental/OnlineServicesEOS/Source/Public/OnlineServicesEOSTypes.h
+++ b/Engine/Plugins/Online/Experimental/OnlineServicesEOS/Source/Public/OnlineServicesEOSTypes.h
@@ -47,52 +47,19 @@ private:
 	}
 };
 
-template<typename TEOSResult, typename TEOSHandle, typename TEOSParameters, typename TEOSFn, typename TAsyncOpType> 
-TFuture<const TEOSResult*> EOS_Async(TOnlineAsyncOp<TAsyncOpType>& Op, TEOSFn EOSFn, TEOSHandle EOSHandle, TEOSParameters Parameters)
+template<typename TEOSResult, typename TEOSHandle, typename TEOSParameters, typename TEOSFn> 
+TFuture<const TEOSResult*> EOS_Async(TEOSFn EOSFn, TEOSHandle EOSHandle, TEOSParameters Parameters)
 {
 	TEOSCallback<TEOSResult>* Callback = new TEOSCallback<TEOSResult>();
 	EOSFn(EOSHandle, &Parameters, Callback, *Callback);
 	return Callback->GetFuture();
 }
 
-template<typename TEOSResult, typename TEOSHandle, typename TEOSParameters, typename TEOSFn, typename TAsyncOpType>
-void EOS_Async(TOnlineAsyncOp<TAsyncOpType>& Op, TEOSFn EOSFn, TEOSHandle EOSHandle, TEOSParameters Parameters, TPromise<const TEOSResult*>&& Promise)
+template<typename TEOSResult, typename TEOSHandle, typename TEOSParameters, typename TEOSFn>
+void EOS_Async(TEOSFn EOSFn, TEOSHandle EOSHandle, TEOSParameters Parameters, TPromise<const TEOSResult*>&& Promise)
 {
 	TEOSCallback<TEOSResult>* Callback = new TEOSCallback<TEOSResult>(MoveTemp(Promise));
 	EOSFn(EOSHandle, &Parameters, Callback, *Callback);
 }
 
-// TEMP until Net Id Registry is done
-extern TMap<EOS_EpicAccountId, int32> EOSAccountIdMap;
-inline FAccountId MakeEOSAccountId(EOS_EpicAccountId EpicAccountId)
-{
-	static int32 EpicAccountIdCounter = 0;
-
-	FAccountId Result;
-	//Result.Type = EOnlineServices::Epic;
-	if (int32* ExistingId = EOSAccountIdMap.Find(EpicAccountId))
-	{
-		Result.Handle = *ExistingId;
-	}
-	else
-	{
-		Result.Handle = EOSAccountIdMap.Emplace(EpicAccountId, ++EpicAccountIdCounter);
-	}
-	return Result;
-}
-
-inline TOptional<EOS_EpicAccountId> EOSAccountIdFromOnlineServiceAccountId(const FAccountId& InAccountId)
-{
-	TOptional<EOS_EpicAccountId> Result;
-	for (const TPair<EOS_EpicAccountId, int32>& EOSAccountIdPair : EOSAccountIdMap)
-	{
-		if (EOSAccountIdPair.Value == InAccountId.Handle)
-		{
-			Result = EOSAccountIdPair.Key;
-			break;
-		}
-	}
-	return Result;
-}
-
 /* UE::Online */ }
diff --git a/Engine/Plugins/Online/Experimental/OnlineServicesEOS/Source/Public/PresenceEOS.h b/Engine/Plugins/Online/Experimental/OnlineServicesEOS/Source/Public/PresenceEOS.h
--- a/Engine/Plugins/Online/Experimental/OnlineServicesEOS/Source/Public/PresenceEOS.h
+++ b/Engine/Plugins/Online/Experimental/OnlineServicesEOS/Source/Public/PresenceEOS.h
@@ -28,13 +28,13 @@ public:
 
 protected:
 	/** Get a user's presence, creating entries if missing */
-	TSharedRef<FUserPresence> FindOrCreatePresence(FAccountId LocalUserId, FAccountId PresenceUserId);
+	TSharedRef<FUserPresence> FindOrCreatePresence(FOnlineAccountIdHandle LocalUserId, FOnlineAccountIdHandle PresenceUserId);
 	/** Update a user's presence from EOS's current value */
-	void UpdateUserPresence(FAccountId LocalUserId, FAccountId PresenceUserId);
+	void UpdateUserPresence(FOnlineAccountIdHandle LocalUserId, FOnlineAccountIdHandle PresenceUserId);
 protected:
 	EOS_HPresence PresenceHandle = nullptr;
 
-	TMap<FAccountId, TMap<FAccountId, TSharedRef<FUserPresence>>> PresenceLists;
+	TMap<FOnlineAccountIdHandle, TMap<FOnlineAccountIdHandle, TSharedRef<FUserPresence>>> PresenceLists;
 	EOS_NotificationId NotifyPresenceChangedNotificationId = 0;
 };
 
