- Fix copy paste typo in Friends update callback
- Improve logging

diff --git a/Engine/Plugins/Online/Experimental/OnlineServicesEOS/Source/Private/FriendsEOS.cpp b/Engine/Plugins/Online/Experimental/OnlineServicesEOS/Source/Private/FriendsEOS.cpp
--- a/Engine/Plugins/Online/Experimental/OnlineServicesEOS/Source/Private/FriendsEOS.cpp
+++ b/Engine/Plugins/Online/Experimental/OnlineServicesEOS/Source/Private/FriendsEOS.cpp
@@ -49,7 +49,7 @@ void FFriendsEOS::Initialize()
 		This->Services.Get<FAuthEOS>()->ResolveAccountId(LocalUserId, Data->TargetUserId)
 		.Next([This, LocalUserId, PreviousStatus = Data->PreviousStatus, CurrentStatus = Data->CurrentStatus](const FOnlineAccountIdHandle& FriendUserId)
 		{
-			UE_LOG(LogTemp, Warning, TEXT("OnEOSPresenceUpdate: [%s] [%s]"), *ToLogString(LocalUserId), *ToLogString(FriendUserId));
+			UE_LOG(LogTemp, Warning, TEXT("OnEOSFriendsUpdate: LocalUserId=[%s] FriendUserId=[%s]"), *ToLogString(LocalUserId), *ToLogString(FriendUserId));
 			This->OnEOSFriendsUpdate(LocalUserId, FriendUserId, PreviousStatus, CurrentStatus);
 		});
 	});
diff --git a/Engine/Plugins/Online/Experimental/OnlineServicesEOS/Source/Private/PresenceEOS.cpp b/Engine/Plugins/Online/Experimental/OnlineServicesEOS/Source/Private/PresenceEOS.cpp
--- a/Engine/Plugins/Online/Experimental/OnlineServicesEOS/Source/Private/PresenceEOS.cpp
+++ b/Engine/Plugins/Online/Experimental/OnlineServicesEOS/Source/Private/PresenceEOS.cpp
@@ -82,7 +82,7 @@ void FPresenceEOS::Initialize()
 		This->Services.Get<FAuthEOS>()->ResolveAccountId(LocalUserId, Data->PresenceUserId)
 		.Next([This, LocalUserId](const FOnlineAccountIdHandle& PresenceUserId)
 		{
-			UE_LOG(LogTemp, Warning, TEXT("OnEOSPresenceUpdate: [%s] [%s]"), *ToLogString(LocalUserId), *ToLogString(PresenceUserId));
+			UE_LOG(LogTemp, Warning, TEXT("OnEOSPresenceUpdate: LocalUserId=[%s] PresenceUserId=[%s]"), *ToLogString(LocalUserId), *ToLogString(PresenceUserId));
 			This->UpdateUserPresence(LocalUserId, PresenceUserId);
 		});
 	});
