Verified logout and login status changed events.
Added fix to avoid crash when logging on and off with the same user.
Added comments to clarify event trigger logic.

diff --git a/Engine/Plugins/Online/Experimental/OnlineServicesEOS/Source/Private/AuthEOS.cpp b/Engine/Plugins/Online/Experimental/OnlineServicesEOS/Source/Private/AuthEOS.cpp
--- a/Engine/Plugins/Online/Experimental/OnlineServicesEOS/Source/Private/AuthEOS.cpp
+++ b/Engine/Plugins/Online/Experimental/OnlineServicesEOS/Source/Private/AuthEOS.cpp
@@ -406,6 +406,7 @@ void FAuthEOS::ProcessSuccessfulLogin(TOnlineAsyncOp<FAuthLogin>& InAsyncOp)
 	TSharedRef<FAccountInfoEOS> AccountInfo = MakeShared<FAccountInfoEOS>();
 	AccountInfo->LocalUserNum = InAsyncOp.GetParams().LocalUserNum;
 	AccountInfo->UserId = LocalUserId;
+	AccountInfo->LoginStatus = ELoginStatus::LoggedIn;
 
 	check(!AccountInfos.Contains(LocalUserId));
 	AccountInfos.Emplace(LocalUserId, AccountInfo);
@@ -413,7 +414,7 @@ void FAuthEOS::ProcessSuccessfulLogin(TOnlineAsyncOp<FAuthLogin>& InAsyncOp)
 	FAuthLogin::Result Result = { AccountInfo };
 	InAsyncOp.SetResult(MoveTemp(Result));
 
-	// Trigger event
+	// When a user logs in, OnEOSLoginStatusChanged can not trigger (if it's that user's first login) or trigger before we add relevant information to AccountInfos, so we trigger the status change event here 
 	FLoginStatusChanged EventParameters;
 	EventParameters.LocalUserId = LocalUserId;
 	EventParameters.PreviousStatus = ELoginStatus::NotLoggedIn;
@@ -474,6 +475,8 @@ TOnlineAsyncOpHandle<FAuthLogout> FAuthEOS::Logout(FAuthLogout::Params&& Params)
 		{
 			// Success
 			InAsyncOp.SetResult(FAuthLogout::Result());
+
+			// OnLoginStatusChanged will be triggered by OnEOSLoginStatusChanged
 		}
 		else
 		{
