- Support updating existing account id's with missing components

diff --git a/Engine/Plugins/Online/Experimental/OnlineServicesEOS/Source/Private/OnlineIdEOS.cpp b/Engine/Plugins/Online/Experimental/OnlineServicesEOS/Source/Private/OnlineIdEOS.cpp
--- a/Engine/Plugins/Online/Experimental/OnlineServicesEOS/Source/Private/OnlineIdEOS.cpp
+++ b/Engine/Plugins/Online/Experimental/OnlineServicesEOS/Source/Private/OnlineIdEOS.cpp
@@ -44,66 +44,93 @@ FOnlineAccountIdHandle FOnlineAccountIdRegistryEOS::Find(const EOS_ProductUserId
 	return Handle;
 }
 
-FOnlineAccountIdHandle FOnlineAccountIdRegistryEOS::Create(const EOS_EpicAccountId EpicAccountId, const EOS_ProductUserId ProductUserId)
+FOnlineAccountIdHandle FOnlineAccountIdRegistryEOS::Create(const EOS_EpicAccountId InEpicAccountId, const EOS_ProductUserId InProductUserId)
 {
 	FOnlineAccountIdHandle Result;
-	const bool bEpicAccountIdValid = EOS_EpicAccountId_IsValid(EpicAccountId) == EOS_TRUE;
-	const bool bProductUserIdValid = EOS_ProductUserId_IsValid(ProductUserId) == EOS_TRUE;
-	check(bEpicAccountIdValid || bProductUserIdValid);
-
-	if (bEpicAccountIdValid || bProductUserIdValid)
+	const bool bInEpicAccountIdValid = EOS_EpicAccountId_IsValid(InEpicAccountId) == EOS_TRUE;
+	const bool bInProductUserIdValid = EOS_ProductUserId_IsValid(InProductUserId) == EOS_TRUE;
+	if (ensure(bInEpicAccountIdValid || bInProductUserIdValid))
 	{
-		auto FindExisting = [&]()
+		bool bUpdateEpicAccountId = false;
+		bool bUpdateProductUserId = false;
+
+		auto FindExisting = [this, InEpicAccountId, InProductUserId, bInEpicAccountIdValid, bInProductUserIdValid, &bUpdateEpicAccountId, &bUpdateProductUserId]()
 		{
 			FOnlineAccountIdHandle Result;
-			if (const FOnlineAccountIdHandle* FoundEas = bEpicAccountIdValid ? EpicAccountIdToHandle.Find(EpicAccountId) : nullptr)
+			if (const FOnlineAccountIdHandle* FoundEas = bInEpicAccountIdValid ? EpicAccountIdToHandle.Find(InEpicAccountId) : nullptr)
 			{
 				Result = *FoundEas;
 			}
-			else if (const FOnlineAccountIdHandle* FoundProd = bProductUserIdValid ? ProductUserIdToHandle.Find(ProductUserId) : nullptr)
+			else if (const FOnlineAccountIdHandle* FoundProd = bInProductUserIdValid ? ProductUserIdToHandle.Find(InProductUserId) : nullptr)
 			{
 				Result = *FoundProd;
 			}
 
 			if (Result.IsValid())
 			{
-				check(EpicAccountId == AccountIdData[Result.GetHandle() - 1].EpicAccountId);
-				check(ProductUserId == AccountIdData[Result.GetHandle() - 1].ProductUserId);
+				const EOS_EpicAccountId FoundEpicAccountId = AccountIdData[Result.GetHandle() - 1].EpicAccountId;
+				const EOS_ProductUserId FoundProductUserId = AccountIdData[Result.GetHandle() - 1].ProductUserId;
+
+				// Check that the found EAS/EOS ids are either unset, or match the input. If a valid input is passed for a currently unset field, this is an update,
+				// which we will track here and complete later under a write lock.
+				check(!FoundEpicAccountId || InEpicAccountId == FoundEpicAccountId);
+				check(!FoundProductUserId || InProductUserId == FoundProductUserId);
+				bUpdateEpicAccountId = !FoundEpicAccountId && bInEpicAccountIdValid;
+				bUpdateProductUserId = !FoundProductUserId && bInProductUserIdValid;
 			}
 
 			return Result;
 		};
 
-		// First take read lock and look for existing elements
 		{
+			// First take read lock and look for existing elements
 			const FReadScopeLock ReadLock(Lock);
 			Result = FindExisting();
 		}
 
 		if(!Result.IsValid())
 		{
-			// Next take write lock, check again for existing elements (in case another thread added it)
+			// Double-checked locking. If we didn't find an element, we take the write lock, and look again, in case another thread raced with us and added one.
 			const FWriteScopeLock WriteLock(Lock);
 			Result = FindExisting();
+
 			if (!Result.IsValid())
 			{
-				// No existing element, so add a new one
+				// We still didn't find one, so now we can add one.
 				FOnlineAccountIdDataEOS& NewAccountIdData = AccountIdData.Emplace_GetRef();
-				NewAccountIdData.EpicAccountId = EpicAccountId;
-				NewAccountIdData.ProductUserId = ProductUserId;
+				NewAccountIdData.EpicAccountId = InEpicAccountId;
+				NewAccountIdData.ProductUserId = InProductUserId;
 
 				Result = FOnlineAccountIdHandle(EOnlineServices::Epic, AccountIdData.Num());
 
-				if (bEpicAccountIdValid)
+				if (bInEpicAccountIdValid)
 				{
-					EpicAccountIdToHandle.Emplace(EpicAccountId, Result);
+					EpicAccountIdToHandle.Emplace(InEpicAccountId, Result);
 				}
-				if (bProductUserIdValid)
+				if (bInProductUserIdValid)
 				{
-					ProductUserIdToHandle.Emplace(ProductUserId, Result);
+					ProductUserIdToHandle.Emplace(InProductUserId, Result);
 				}
 			}
 		}
+
+		check(Result.IsValid());
+		if (bUpdateEpicAccountId || bUpdateProductUserId)
+		{
+			// Finally, update any previously unset fields for which we now have a valid value.
+			const FWriteScopeLock WriteLock(Lock);
+			FOnlineAccountIdDataEOS& AccountIdDataToUpdate = AccountIdData[Result.GetHandle() -1];
+			if (bUpdateEpicAccountId)
+			{
+				AccountIdDataToUpdate.EpicAccountId = InEpicAccountId;
+				EpicAccountIdToHandle.Emplace(InEpicAccountId, Result);
+			}
+			if (bUpdateProductUserId)
+			{
+				AccountIdDataToUpdate.ProductUserId = InProductUserId;
+				ProductUserIdToHandle.Emplace(InProductUserId, Result);
+			}
+		}
 	}
 	
 	return Result;
