Added logic to process external authentication tokens in auth interface.
Added logic to process account linking on a failed auth login result.

diff --git a/Engine/Plugins/Online/Experimental/OnlineServices/Source/OnlineServicesCommon/Public/Online/OnlineExecHandler.h b/Engine/Plugins/Online/Experimental/OnlineServices/Source/OnlineServicesCommon/Public/Online/OnlineExecHandler.h
--- a/Engine/Plugins/Online/Experimental/OnlineServices/Source/OnlineServicesCommon/Public/Online/OnlineExecHandler.h
+++ b/Engine/Plugins/Online/Experimental/OnlineServices/Source/OnlineServicesCommon/Public/Online/OnlineExecHandler.h
@@ -67,6 +67,21 @@ inline bool ParseOnlineExecParams(const TCHAR*& Cmd, FString& Value)
 	return true;
 }
 
+inline bool ParseOnlineExecParams(const TCHAR*& Cmd, uint8& Value)
+{
+	FString Token;
+	if (FParse::Token(Cmd, Token, true))
+	{
+		Value = static_cast<uint8>(FCString::Strtoui64(Cmd, nullptr, 10));
+	}
+	else
+	{
+		return false;
+	}
+
+	return true;
+}
+
 inline bool ParseOnlineExecParams(const TCHAR*& Cmd, int32& Value)
 {
 	FString Token;
@@ -153,6 +168,13 @@ inline bool ParseOnlineExecParams(const TCHAR*& Cmd, TArray<T>& Array)
 	return true;
 }
 
+template <typename... Ts>
+inline bool ParseOnlineExecParams(const TCHAR*& Cmd, TVariant<Ts...>& Variant)
+{
+	// TODO: This is a temporary stub implementation
+	return true;
+}
+
 template <typename T>
 std::enable_if_t<!TModels<Meta::COnlineMetadataAvailable, T>::Value, bool> ParseOnlineExecParams(const TCHAR*& Cmd, T& Value)
 {
diff --git a/Engine/Plugins/Online/Experimental/OnlineServices/Source/OnlineServicesInterface/Public/Online/Auth.h b/Engine/Plugins/Online/Experimental/OnlineServices/Source/OnlineServicesInterface/Public/Online/Auth.h
--- a/Engine/Plugins/Online/Experimental/OnlineServices/Source/OnlineServicesInterface/Public/Online/Auth.h
+++ b/Engine/Plugins/Online/Experimental/OnlineServices/Source/OnlineServicesInterface/Public/Online/Auth.h
@@ -34,6 +34,8 @@ public:
 	// TODO: Other fields
 };
 
+using FCredentialsToken = TVariant<FString, TArray<uint8>>;
+
 struct FAuthLogin
 {
 	static constexpr TCHAR Name[] = TEXT("Login");
@@ -43,7 +45,7 @@ struct FAuthLogin
 		int32 LocalUserNum;
 		FString CredentialsType;
 		FString CredentialsId;
-		FString CredentialsToken;
+		FCredentialsToken CredentialsToken;
 		TArray<FString> Scopes;
 	};
 
diff --git a/Engine/Plugins/Online/Experimental/OnlineServices/Source/OnlineServicesInterface/Public/Online/OnlineUtils.h b/Engine/Plugins/Online/Experimental/OnlineServices/Source/OnlineServicesInterface/Public/Online/OnlineUtils.h
--- a/Engine/Plugins/Online/Experimental/OnlineServices/Source/OnlineServicesInterface/Public/Online/OnlineUtils.h
+++ b/Engine/Plugins/Online/Experimental/OnlineServices/Source/OnlineServicesInterface/Public/Online/OnlineUtils.h
@@ -83,6 +83,15 @@ FString ToLogString(const TOptional<T> Optional)
 	}
 }
 
+template <typename... Ts>
+FString ToLogString(const TVariant<Ts...>& Variant)
+{
+	return Visit([](const auto& Value)
+	{
+		return ToLogString(Value);
+	}, Variant);
+}
+
 template <typename T, ESPMode Mode>
 FString ToLogString(const TSharedRef<T, Mode>& Ref)
 {
diff --git a/Engine/Plugins/Online/Experimental/OnlineServicesEOS/Source/Private/AuthEOS.cpp b/Engine/Plugins/Online/Experimental/OnlineServicesEOS/Source/Private/AuthEOS.cpp
--- a/Engine/Plugins/Online/Experimental/OnlineServicesEOS/Source/Private/AuthEOS.cpp
+++ b/Engine/Plugins/Online/Experimental/OnlineServicesEOS/Source/Private/AuthEOS.cpp
@@ -96,6 +96,65 @@ struct FEOSAuthCredentials :
 		return *this;
 	}
 
+	void inline SetExternalCredentialType(EOnlineServices InPlatform, const FString& InCredentialsId)
+	{
+		if (InPlatform == EOnlineServices::Steam)
+		{
+			ExternalType = EOS_EExternalCredentialType::EOS_ECT_STEAM_APP_TICKET;
+		}
+		else if (InPlatform == EOnlineServices::PSN)
+		{
+			ExternalType = EOS_EExternalCredentialType::EOS_ECT_PSN_ID_TOKEN;
+		}
+		else if (InPlatform == EOnlineServices::Xbox)
+		{
+			ExternalType = EOS_EExternalCredentialType::EOS_ECT_XBL_XSTS_TOKEN;
+		}
+		else if (InPlatform == EOnlineServices::Nintendo)
+		{
+			if (InCredentialsId == TEXT("NintendoAccount"))
+			{
+				ExternalType = EOS_EExternalCredentialType::EOS_ECT_NINTENDO_ID_TOKEN;
+			}
+			else
+			{
+				ExternalType = EOS_EExternalCredentialType::EOS_ECT_NINTENDO_NSA_ID_TOKEN;
+			}
+		}
+		else if (InPlatform == EOnlineServices::Apple)
+		{
+			ExternalType = EOS_EExternalCredentialType::EOS_ECT_APPLE_ID_TOKEN;
+		}
+		else if (InPlatform == EOnlineServices::Google)
+		{
+			ExternalType = EOS_EExternalCredentialType::EOS_ECT_GOOGLE_ID_TOKEN;
+		}
+		else
+		{
+			// Unknown means OpenID
+			ExternalType = EOS_EExternalCredentialType::EOS_ECT_OPENID_ACCESS_TOKEN;
+		}
+	}
+
+	void SetToken(const FCredentialsToken& InToken)
+	{
+		if (InToken.IsType<TArray<uint8>>())
+		{
+			const TArray<uint8>& TokenData = InToken.Get<TArray<uint8>>();
+			uint32_t InOutBufferLength = EOS_OSS_STRING_BUFFER_LENGTH;
+			EOS_ByteArray_ToString(TokenData.GetData(), TokenData.Num(), TokenAnsi, &InOutBufferLength);
+		}
+		else if (InToken.IsType<FString>())
+		{
+			const FString& TokenString = InToken.Get<FString>();
+			FCStringAnsi::Strncpy(TokenAnsi, TCHAR_TO_UTF8(*TokenString), TokenString.Len() + 1);
+		}
+		else
+		{
+			UE_LOG(LogTemp, Warning, TEXT("FEOSAuthCredentials::Token cannot be set with an invalid credentials token parameter. Please ensure there is valid data in the credentials token."));
+		}
+	}
+
 	char IdAnsi[EOS_OSS_STRING_BUFFER_LENGTH];
 	char TokenAnsi[EOS_MAX_TOKEN_SIZE];
 };
@@ -247,19 +306,23 @@ TOnlineAsyncOpHandle<FAuthLogin> FAuthEOS::Login(FAuthLogin::Params&& Params)
 	{
 		switch (Credentials.Type)
 		{
+		case EOS_ELoginCredentialType::EOS_LCT_ExternalAuth:
+			Credentials.SetExternalCredentialType(EOnlineServices::Steam, Op->GetParams().CredentialsId); // TODO: Will be GetServicesProvider() or something similar
+			Credentials.SetToken(Op->GetParams().CredentialsToken);
+			break;
 		case EOS_ELoginCredentialType::EOS_LCT_ExchangeCode:
 			// This is how the Epic launcher will pass credentials to you
 			Credentials.IdAnsi[0] = '\0';
-			FCStringAnsi::Strncpy(Credentials.TokenAnsi, TCHAR_TO_UTF8(*Op->GetParams().CredentialsToken), EOS_MAX_TOKEN_SIZE);
+			FCStringAnsi::Strncpy(Credentials.TokenAnsi, TCHAR_TO_UTF8(*Op->GetParams().CredentialsToken.Get<FString>()), EOS_MAX_TOKEN_SIZE);
 			break;
 		case EOS_ELoginCredentialType::EOS_LCT_Password:
 			FCStringAnsi::Strncpy(Credentials.IdAnsi, TCHAR_TO_UTF8(*Op->GetParams().CredentialsId), EOS_OSS_STRING_BUFFER_LENGTH);
-			FCStringAnsi::Strncpy(Credentials.TokenAnsi, TCHAR_TO_UTF8(*Op->GetParams().CredentialsToken), EOS_MAX_TOKEN_SIZE);
+			FCStringAnsi::Strncpy(Credentials.TokenAnsi, TCHAR_TO_UTF8(*Op->GetParams().CredentialsToken.Get<FString>()), EOS_MAX_TOKEN_SIZE);
 			break;
 		case EOS_ELoginCredentialType::EOS_LCT_Developer:
 			// This is auth via the EOS auth tool
 			FCStringAnsi::Strncpy(Credentials.IdAnsi, TCHAR_TO_UTF8(*Op->GetParams().CredentialsId), EOS_OSS_STRING_BUFFER_LENGTH);
-			FCStringAnsi::Strncpy(Credentials.TokenAnsi, TCHAR_TO_UTF8(*Op->GetParams().CredentialsToken), EOS_MAX_TOKEN_SIZE);
+			FCStringAnsi::Strncpy(Credentials.TokenAnsi, TCHAR_TO_UTF8(*Op->GetParams().CredentialsToken.Get<FString>()), EOS_MAX_TOKEN_SIZE);
 			break;
 		case EOS_ELoginCredentialType::EOS_LCT_AccountPortal:
 			// This is auth via the EOS Account Portal
@@ -298,12 +361,56 @@ TOnlineAsyncOpHandle<FAuthLogin> FAuthEOS::Login(FAuthLogin::Params&& Params)
 			// We cache the Epic Account Id to use it in later stages of the login process
 			InAsyncOp.Data.Set(TEXT("EpicAccountId"), Data->LocalUserId);
 
+			// A success means the account is already linked, we'll process the NULL result in the next step
+		}
+		else if (Data->ResultCode == EOS_EResult::EOS_InvalidUser && Data->ContinuanceToken != nullptr)
+		{
+			EOS_Auth_LinkAccountOptions LinkAccountOptions = {};
+			LinkAccountOptions.ApiVersion = EOS_AUTH_LINKACCOUNT_API_LATEST;
+			LinkAccountOptions.ContinuanceToken = Data->ContinuanceToken;
+
+			return EOS_Async<EOS_Auth_LinkAccountCallbackInfo>(EOS_Auth_LinkAccount, AuthHandle, LinkAccountOptions);
+		}
+		else
+		{
+			FOnlineError Error = Errors::Unknown();
+			if (Data->ResultCode == EOS_EResult::EOS_InvalidAuth)
+			{
+				Error = Errors::InvalidCreds();
+			}
+
+			InAsyncOp.SetError(MoveTemp(Error));
+		}
+
+		return MakeFulfilledPromise<const EOS_Auth_LinkAccountCallbackInfo*>().GetFuture();
+	})
+	.Then([this](TOnlineAsyncOp<FAuthLogin>& InAsyncOp, const EOS_Auth_LinkAccountCallbackInfo* Data)
+	{
+		UE_LOG(LogTemp, Verbose, TEXT("[FAuthEOS::Login] EOS_Auth_LinkAccount Result: [%s]"), *LexToString(Data->ResultCode));
+		
+		// If Data is NULL, it means the account was already linked
+		if (Data == nullptr || Data->ResultCode == EOS_EResult::EOS_Success)
+		{
+			EOS_EpicAccountId EpicAccountId = nullptr;
+
+			if (Data == NULL)
+			{
+				EpicAccountId = *InAsyncOp.Data.Get<EOS_EpicAccountId>(TEXT("EpicAccountId"));
+			}
+			else
+			{
+				// We cache the Epic Account Id to use it in later stages of the login process
+				InAsyncOp.Data.Set(TEXT("EpicAccountId"), Data->LocalUserId);
+
+				EpicAccountId = Data->LocalUserId;
+			}
+
 			// On success, attempt Connect Login
 			EOS_Auth_Token* AuthToken = nullptr;
 			EOS_Auth_CopyUserAuthTokenOptions CopyOptions = { };
 			CopyOptions.ApiVersion = EOS_AUTH_COPYUSERAUTHTOKEN_API_LATEST;
 
-			EOS_EResult CopyResult = EOS_Auth_CopyUserAuthToken(AuthHandle, &CopyOptions, Data->LocalUserId, &AuthToken);
+			EOS_EResult CopyResult = EOS_Auth_CopyUserAuthToken(AuthHandle, &CopyOptions, EpicAccountId, &AuthToken);
 
 			UE_LOG(LogTemp, Verbose, TEXT("[FAuthEOS::Login] EOS_Auth_CopyUserAuthToken Result: [%s]"), *LexToString(CopyResult));
 
@@ -327,19 +434,9 @@ TOnlineAsyncOpHandle<FAuthLogin> FAuthEOS::Login(FAuthLogin::Params&& Params)
 				InAsyncOp.SetError(Errors::Unknown()); // TODO
 			}
 		}
-		else if (Data->ResultCode == EOS_EResult::EOS_InvalidUser && Data->ContinuanceToken != nullptr)
-		{
-			// Link Account
-		}
 		else
 		{
-			FOnlineError Error = Errors::Unknown();
-			if (Data->ResultCode == EOS_EResult::EOS_InvalidAuth)
-			{
-				Error = Errors::InvalidCreds();
-			}
-
-			InAsyncOp.SetError(MoveTemp(Error));
+			InAsyncOp.SetError(Errors::Unknown()); // TODO
 		}
 
 		return MakeFulfilledPromise<const EOS_Connect_LoginCallbackInfo*>().GetFuture();
