Add ShowLoginUI to External UI
Change int LocalUserNum to FPlatformUserId
Add display name to account

diff --git a/Engine/Plugins/Online/Experimental/OnlineServices/Source/OnlineServicesCommon/Private/Online/AuthCommon.cpp b/Engine/Plugins/Online/Experimental/OnlineServices/Source/OnlineServicesCommon/Private/Online/AuthCommon.cpp
--- a/Engine/Plugins/Online/Experimental/OnlineServices/Source/OnlineServicesCommon/Private/Online/AuthCommon.cpp
+++ b/Engine/Plugins/Online/Experimental/OnlineServices/Source/OnlineServicesCommon/Private/Online/AuthCommon.cpp
@@ -18,7 +18,7 @@ void FAuthCommon::RegisterCommands()
 	RegisterCommand(&FAuthCommon::Login);
 	RegisterCommand(&FAuthCommon::Logout);
 	RegisterCommand(&FAuthCommon::GenerateAuth);
-	RegisterCommand(&FAuthCommon::GetAccountByLocalUserNum);
+	RegisterCommand(&FAuthCommon::GetAccountByPlatformUserId);
 	RegisterCommand(&FAuthCommon::GetAccountByAccountId);
 }
 
@@ -43,9 +43,9 @@ TOnlineAsyncOpHandle<FAuthGenerateAuth> FAuthCommon::GenerateAuth(FAuthGenerateA
 	return Operation->GetHandle();
 }
 
-TOnlineResult<FAuthGetAccountByLocalUserNum> FAuthCommon::GetAccountByLocalUserNum(FAuthGetAccountByLocalUserNum::Params&& Params)
+TOnlineResult<FAuthGetAccountByPlatformUserId> FAuthCommon::GetAccountByPlatformUserId(FAuthGetAccountByPlatformUserId::Params&& Params)
 {
-	return TOnlineResult<FAuthGetAccountByLocalUserNum>(Errors::NotImplemented());
+	return TOnlineResult<FAuthGetAccountByPlatformUserId>(Errors::NotImplemented());
 }
 
 TOnlineResult<FAuthGetAccountByAccountId> FAuthCommon::GetAccountByAccountId(FAuthGetAccountByAccountId::Params&& Params)
diff --git a/Engine/Plugins/Online/Experimental/OnlineServices/Source/OnlineServicesCommon/Private/Online/ExternalUICommon.cpp b/Engine/Plugins/Online/Experimental/OnlineServices/Source/OnlineServicesCommon/Private/Online/ExternalUICommon.cpp
--- a/Engine/Plugins/Online/Experimental/OnlineServices/Source/OnlineServicesCommon/Private/Online/ExternalUICommon.cpp
+++ b/Engine/Plugins/Online/Experimental/OnlineServices/Source/OnlineServicesCommon/Private/Online/ExternalUICommon.cpp
@@ -18,6 +18,13 @@ void FExternalUICommon::RegisterCommands()
 	RegisterCommand(&FExternalUICommon::ShowFriendsUI);
 }
 
+TOnlineAsyncOpHandle<FExternalUIShowLoginUI> FExternalUICommon::ShowLoginUI(FExternalUIShowLoginUI::Params&& Params)
+{
+	TOnlineAsyncOpRef<FExternalUIShowLoginUI> Operation = GetOp<FExternalUIShowLoginUI>(MoveTemp(Params));
+	Operation->SetError(Errors::NotImplemented());
+	return Operation->GetHandle();
+}
+
 TOnlineAsyncOpHandle<FExternalUIShowFriendsUI> FExternalUICommon::ShowFriendsUI(FExternalUIShowFriendsUI::Params&& Params)
 {
 	TOnlineAsyncOpRef<FExternalUIShowFriendsUI> Operation = GetOp<FExternalUIShowFriendsUI>(MoveTemp(Params));
diff --git a/Engine/Plugins/Online/Experimental/OnlineServices/Source/OnlineServicesCommon/Private/Online/OnlineSyntaxTests.cpp b/Engine/Plugins/Online/Experimental/OnlineServices/Source/OnlineServicesCommon/Private/Online/OnlineSyntaxTests.cpp
--- a/Engine/Plugins/Online/Experimental/OnlineServices/Source/OnlineServicesCommon/Private/Online/OnlineSyntaxTests.cpp
+++ b/Engine/Plugins/Online/Experimental/OnlineServices/Source/OnlineServicesCommon/Private/Online/OnlineSyntaxTests.cpp
@@ -448,8 +448,8 @@ void TestToLogStringSyntax()
 	ToLogString(UE::Online::FAuthLogout::Result());
 	ToLogString(UE::Online::FAuthGenerateAuth::Params());
 	ToLogString(UE::Online::FAuthGenerateAuth::Result());
-	ToLogString(UE::Online::FAuthGetAccountByLocalUserNum::Params());
-	ToLogString(UE::Online::FAuthGetAccountByLocalUserNum::Result());
+	ToLogString(UE::Online::FAuthGetAccountByPlatformUserId::Params());
+	ToLogString(UE::Online::FAuthGetAccountByPlatformUserId::Result());
 	ToLogString(UE::Online::FAuthGetAccountByAccountId::Params());
 	ToLogString(UE::Online::FAuthGetAccountByAccountId::Result());
 	ToLogString(UE::Online::FLoginStatusChanged());
diff --git a/Engine/Plugins/Online/Experimental/OnlineServices/Source/OnlineServicesCommon/Public/Online/AuthCommon.h b/Engine/Plugins/Online/Experimental/OnlineServices/Source/OnlineServicesCommon/Public/Online/AuthCommon.h
--- a/Engine/Plugins/Online/Experimental/OnlineServices/Source/OnlineServicesCommon/Public/Online/AuthCommon.h
+++ b/Engine/Plugins/Online/Experimental/OnlineServices/Source/OnlineServicesCommon/Public/Online/AuthCommon.h
@@ -23,7 +23,7 @@ public:
 	virtual TOnlineAsyncOpHandle<FAuthLogin> Login(FAuthLogin::Params&& Params) override;
 	virtual TOnlineAsyncOpHandle<FAuthLogout> Logout(FAuthLogout::Params&& Params) override;
 	virtual TOnlineAsyncOpHandle<FAuthGenerateAuth> GenerateAuth(FAuthGenerateAuth::Params&& Params) override;
-	virtual TOnlineResult<FAuthGetAccountByLocalUserNum> GetAccountByLocalUserNum(FAuthGetAccountByLocalUserNum::Params&& Params) override;
+	virtual TOnlineResult<FAuthGetAccountByPlatformUserId> GetAccountByPlatformUserId(FAuthGetAccountByPlatformUserId::Params&& Params) override;
 	virtual TOnlineResult<FAuthGetAccountByAccountId> GetAccountByAccountId(FAuthGetAccountByAccountId::Params&& Params) override;
 	virtual TOnlineEvent<void(const FLoginStatusChanged&)> OnLoginStatusChanged() override;
 
diff --git a/Engine/Plugins/Online/Experimental/OnlineServices/Source/OnlineServicesCommon/Public/Online/ExternalUICommon.h b/Engine/Plugins/Online/Experimental/OnlineServices/Source/OnlineServicesCommon/Public/Online/ExternalUICommon.h
--- a/Engine/Plugins/Online/Experimental/OnlineServices/Source/OnlineServicesCommon/Public/Online/ExternalUICommon.h
+++ b/Engine/Plugins/Online/Experimental/OnlineServices/Source/OnlineServicesCommon/Public/Online/ExternalUICommon.h
@@ -20,6 +20,7 @@ public:
 	virtual void RegisterCommands() override;
 
 	// IExternalUI
+	virtual TOnlineAsyncOpHandle<FExternalUIShowLoginUI> ShowLoginUI(FExternalUIShowLoginUI::Params&& Params) override;
 	virtual TOnlineAsyncOpHandle<FExternalUIShowFriendsUI> ShowFriendsUI(FExternalUIShowFriendsUI::Params&& Params) override;
 };
 
diff --git a/Engine/Plugins/Online/Experimental/OnlineServices/Source/OnlineServicesCommon/Public/Online/OnlineExecHandler.h b/Engine/Plugins/Online/Experimental/OnlineServices/Source/OnlineServicesCommon/Public/Online/OnlineExecHandler.h
--- a/Engine/Plugins/Online/Experimental/OnlineServices/Source/OnlineServicesCommon/Public/Online/OnlineExecHandler.h
+++ b/Engine/Plugins/Online/Experimental/OnlineServices/Source/OnlineServicesCommon/Public/Online/OnlineExecHandler.h
@@ -142,6 +142,21 @@ inline bool ParseOnlineExecParams(const TCHAR*& Cmd, uint64& Value)
 	return true;
 }
 
+inline bool ParseOnlineExecParams(const TCHAR*& Cmd, FPlatformUserId& Value)
+{
+	int32 InternalId = -1;
+	if (ParseOnlineExecParams(Cmd, InternalId))
+	{
+		Value.CreateFromInternalId(InternalId);
+	}
+	else
+	{
+		return false;
+	}
+
+	return true;
+}
+
 template <typename T>
 inline bool ParseOnlineExecParams(const TCHAR*& Cmd, TArray<T>& Array)
 {
diff --git a/Engine/Plugins/Online/Experimental/OnlineServices/Source/OnlineServicesInterface/Public/Online/Auth.h b/Engine/Plugins/Online/Experimental/OnlineServices/Source/OnlineServicesInterface/Public/Online/Auth.h
--- a/Engine/Plugins/Online/Experimental/OnlineServices/Source/OnlineServicesInterface/Public/Online/Auth.h
+++ b/Engine/Plugins/Online/Experimental/OnlineServices/Source/OnlineServicesInterface/Public/Online/Auth.h
@@ -10,7 +10,7 @@ namespace UE::Online {
 
 class FOnlineError;
 
-enum class ELoginStatus
+enum class ELoginStatus : uint8
 {
 	/** Player has not logged in or chosen a local profile */
 	NotLoggedIn,
@@ -26,11 +26,13 @@ class FAccountInfo
 {
 public:
 	/** Local user num */
-	int32 LocalUserNum;
+	FPlatformUserId PlatformUserId;
 	/** Account id */
 	FOnlineAccountIdHandle UserId;
 	/** Login status */
 	ELoginStatus LoginStatus;
+	/** Display name */
+	FString DisplayName;
 	// TODO: Other fields
 };
 
@@ -42,7 +44,7 @@ struct FAuthLogin
 
 	struct Params
 	{
-		int32 LocalUserNum;
+		FPlatformUserId PlatformUserId = PLATFORMUSERID_NONE;
 		FString CredentialsType;
 		FString CredentialsId;
 		FCredentialsToken CredentialsToken;
@@ -86,13 +88,13 @@ struct FAuthGenerateAuth
 	};
 };
 
-struct FAuthGetAccountByLocalUserNum
+struct FAuthGetAccountByPlatformUserId
 {
-	static constexpr TCHAR Name[] = TEXT("GetAccountByLocalUserNum");
+	static constexpr TCHAR Name[] = TEXT("GetAccountByPlatformUserId");
 
 	struct Params
 	{
-		int32 LocalUserNum;
+		FPlatformUserId PlatformUserId = PLATFORMUSERID_NONE;
 	};
 
 	struct Result
@@ -148,7 +150,7 @@ public:
 	/**
 	 * Get logged in user by local user num
 	 */ 
-	virtual TOnlineResult<FAuthGetAccountByLocalUserNum> GetAccountByLocalUserNum(FAuthGetAccountByLocalUserNum::Params&& Params) = 0;
+	virtual TOnlineResult<FAuthGetAccountByPlatformUserId> GetAccountByPlatformUserId(FAuthGetAccountByPlatformUserId::Params&& Params) = 0;
 
 	/**
 	 * Get logged in user by account id
@@ -166,13 +168,13 @@ namespace Meta {
 // TODO: Move to Auth_Meta.inl file?
 
 BEGIN_ONLINE_STRUCT_META(FAccountInfo)
-	ONLINE_STRUCT_FIELD(FAccountInfo, LocalUserNum),
+	ONLINE_STRUCT_FIELD(FAccountInfo, PlatformUserId),
 	ONLINE_STRUCT_FIELD(FAccountInfo, UserId),
 	ONLINE_STRUCT_FIELD(FAccountInfo, LoginStatus)
 END_ONLINE_STRUCT_META()
 
 BEGIN_ONLINE_STRUCT_META(FAuthLogin::Params)
-	ONLINE_STRUCT_FIELD(FAuthLogin::Params, LocalUserNum),
+	ONLINE_STRUCT_FIELD(FAuthLogin::Params, PlatformUserId),
 	ONLINE_STRUCT_FIELD(FAuthLogin::Params, CredentialsType),
 	ONLINE_STRUCT_FIELD(FAuthLogin::Params, CredentialsId),
 	ONLINE_STRUCT_FIELD(FAuthLogin::Params, CredentialsToken),
@@ -199,12 +201,12 @@ END_ONLINE_STRUCT_META()
 BEGIN_ONLINE_STRUCT_META(FAuthGenerateAuth::Result)
 END_ONLINE_STRUCT_META()
 
-BEGIN_ONLINE_STRUCT_META(FAuthGetAccountByLocalUserNum::Params)
-	ONLINE_STRUCT_FIELD(FAuthGetAccountByLocalUserNum::Params, LocalUserNum)
+BEGIN_ONLINE_STRUCT_META(FAuthGetAccountByPlatformUserId::Params)
+	ONLINE_STRUCT_FIELD(FAuthGetAccountByPlatformUserId::Params, PlatformUserId)
 END_ONLINE_STRUCT_META()
 
-BEGIN_ONLINE_STRUCT_META(FAuthGetAccountByLocalUserNum::Result)
-	ONLINE_STRUCT_FIELD(FAuthGetAccountByLocalUserNum::Result, AccountInfo)
+BEGIN_ONLINE_STRUCT_META(FAuthGetAccountByPlatformUserId::Result)
+	ONLINE_STRUCT_FIELD(FAuthGetAccountByPlatformUserId::Result, AccountInfo)
 END_ONLINE_STRUCT_META()
 
 BEGIN_ONLINE_STRUCT_META(FAuthGetAccountByAccountId::Params)
diff --git a/Engine/Plugins/Online/Experimental/OnlineServices/Source/OnlineServicesInterface/Public/Online/ExternalUI.h b/Engine/Plugins/Online/Experimental/OnlineServices/Source/OnlineServicesInterface/Public/Online/ExternalUI.h
--- a/Engine/Plugins/Online/Experimental/OnlineServices/Source/OnlineServicesInterface/Public/Online/ExternalUI.h
+++ b/Engine/Plugins/Online/Experimental/OnlineServices/Source/OnlineServicesInterface/Public/Online/ExternalUI.h
@@ -9,6 +9,26 @@
 namespace UE::Online {
 
 class FOnlineError;
+class FAccountInfo;
+
+struct FExternalUIShowLoginUI
+{
+	static constexpr TCHAR Name[] = TEXT("ShowLoginUI");
+
+	struct Params
+	{
+		/** Platform user id logging in */
+		FPlatformUserId PlatformUserId = PLATFORMUSERID_NONE;
+		/** Auth scopes to request */
+		TArray<FString> Scopes;
+	};
+
+	struct Result
+	{
+		/** The account info that was logged into */
+		TSharedRef<FAccountInfo> AccountInfo;
+	};
+};
 
 struct FExternalUIShowFriendsUI
 {
@@ -16,6 +36,7 @@ struct FExternalUIShowFriendsUI
 
 	struct Params
 	{
+		/** Local user id to show the friends UI for */
 		FOnlineAccountIdHandle LocalUserId;
 	};
 
@@ -28,7 +49,12 @@ class IExternalUI
 {
 public:
 	/**
-	 * Shows the friends tab of the EOS Overlay if enabled
+	 * Shows the online service's login UI
+	 */
+	virtual TOnlineAsyncOpHandle<FExternalUIShowLoginUI> ShowLoginUI(FExternalUIShowLoginUI::Params&& Params) = 0;
+
+	/**
+	 * Shows the online service's friends UI
 	 */
 	virtual TOnlineAsyncOpHandle<FExternalUIShowFriendsUI> ShowFriendsUI(FExternalUIShowFriendsUI::Params&& Params) = 0;
 };
@@ -36,6 +62,15 @@ public:
 namespace Meta {
 // TODO: Move to ExternalUI_Meta.inl file?
 
+BEGIN_ONLINE_STRUCT_META(FExternalUIShowLoginUI::Params)
+	ONLINE_STRUCT_FIELD(FExternalUIShowLoginUI::Params, PlatformUserId),
+	ONLINE_STRUCT_FIELD(FExternalUIShowLoginUI::Params, Scopes)
+END_ONLINE_STRUCT_META()
+
+BEGIN_ONLINE_STRUCT_META(FExternalUIShowLoginUI::Result)
+	ONLINE_STRUCT_FIELD(FExternalUIShowLoginUI::Result, AccountInfo)
+END_ONLINE_STRUCT_META()
+
 BEGIN_ONLINE_STRUCT_META(FExternalUIShowFriendsUI::Params)
 	ONLINE_STRUCT_FIELD(FExternalUIShowFriendsUI::Params, LocalUserId)
 END_ONLINE_STRUCT_META()
diff --git a/Engine/Plugins/Online/Experimental/OnlineServices/Source/OnlineServicesInterface/Public/Online/OnlineServices.h b/Engine/Plugins/Online/Experimental/OnlineServices/Source/OnlineServicesInterface/Public/Online/OnlineServices.h
--- a/Engine/Plugins/Online/Experimental/OnlineServices/Source/OnlineServicesInterface/Public/Online/OnlineServices.h
+++ b/Engine/Plugins/Online/Experimental/OnlineServices/Source/OnlineServicesInterface/Public/Online/OnlineServices.h
@@ -10,6 +10,7 @@ class FString;
 namespace UE::Online {
 
 // Interfaces
+using IOnlineServicesPtr = TSharedPtr<class IOnlineServices>;
 using IAuthPtr = TSharedPtr<class IAuth>;
 using IFriendsPtr = TSharedPtr<class IFriends>;
 using IPresencePtr = TSharedPtr<class IPresence>;
diff --git a/Engine/Plugins/Online/Experimental/OnlineServices/Source/OnlineServicesInterface/Public/Online/OnlineUtils.h b/Engine/Plugins/Online/Experimental/OnlineServices/Source/OnlineServicesInterface/Public/Online/OnlineUtils.h
--- a/Engine/Plugins/Online/Experimental/OnlineServices/Source/OnlineServicesInterface/Public/Online/OnlineUtils.h
+++ b/Engine/Plugins/Online/Experimental/OnlineServices/Source/OnlineServicesInterface/Public/Online/OnlineUtils.h
@@ -143,6 +143,11 @@ inline FString ToLogString(int64 Value)
 	return FString::Printf(TEXT("%lli"), Value);
 }
 
+inline FString ToLogString(FPlatformUserId PlatformUserId)
+{
+	return ToLogString(PlatformUserId.GetInternalId());
+}
+
 template<typename T>
 std::enable_if_t<!TModels<Meta::COnlineMetadataAvailable, T>::Value, FString> ToLogString(const T& Value)
 {
diff --git a/Engine/Plugins/Online/Experimental/OnlineServicesEOS/Source/Private/AuthEOS.cpp b/Engine/Plugins/Online/Experimental/OnlineServicesEOS/Source/Private/AuthEOS.cpp
--- a/Engine/Plugins/Online/Experimental/OnlineServicesEOS/Source/Private/AuthEOS.cpp
+++ b/Engine/Plugins/Online/Experimental/OnlineServicesEOS/Source/Private/AuthEOS.cpp
@@ -9,6 +9,7 @@
 #include "Online/AuthErrors.h"
 #include "Online/OnlineAsyncOp.h"
 #include "Online/OnlineErrorDefinitions.h"
+#include "Misc/CommandLine.h"
 
 #include "eos_auth.h"
 #include "eos_common.h"
@@ -17,6 +18,7 @@
 #include "eos_init.h"
 #include "eos_sdk.h"
 #include "eos_logging.h"
+#include "eos_userinfo.h"
 
 namespace UE::Online {
 
@@ -272,6 +274,22 @@ bool LexFromString(EOS_ELoginCredentialType& OutEnum, const TCHAR* const InStrin
 
 TOnlineAsyncOpHandle<FAuthLogin> FAuthEOS::Login(FAuthLogin::Params&& Params)
 {
+	// Is this auto-login?
+	if (Params.CredentialsId.IsEmpty() && Params.CredentialsType.IsEmpty() && Params.CredentialsToken.IsType<FString>() && Params.CredentialsToken.Get<FString>().IsEmpty())
+	{
+		FString	CommandLineAuthId;
+		FString	CommandLineAuthToken;
+		FString	CommandLineAuthType;
+		FParse::Value(FCommandLine::Get(), TEXT("AUTH_LOGIN="), CommandLineAuthId);
+		FParse::Value(FCommandLine::Get(), TEXT("AUTH_PASSWORD="), CommandLineAuthToken);
+		FParse::Value(FCommandLine::Get(), TEXT("AUTH_TYPE="), CommandLineAuthType);
+		if (!CommandLineAuthId.IsEmpty() && !CommandLineAuthToken.IsEmpty() && !CommandLineAuthType.IsEmpty())
+		{
+			Params.CredentialsId = MoveTemp(CommandLineAuthId);
+			Params.CredentialsToken.Emplace<FString>(MoveTemp(CommandLineAuthToken));
+			Params.CredentialsType = MoveTemp(CommandLineAuthType);
+		}
+	}
 	TOnlineAsyncOpRef<FAuthLogin> Op = GetOp<FAuthLogin>(MoveTemp(Params));
 
 	EOS_Auth_LoginOptions LoginOptions = { };
@@ -501,10 +519,28 @@ void FAuthEOS::ProcessSuccessfulLogin(TOnlineAsyncOp<FAuthLogin>& InAsyncOp)
 	UE_LOG(LogTemp, Verbose, TEXT("[FAuthEOS::Login] Successfully logged in as [%s]"), *ToLogString(LocalUserId));
 
 	TSharedRef<FAccountInfoEOS> AccountInfo = MakeShared<FAccountInfoEOS>();
-	AccountInfo->LocalUserNum = InAsyncOp.GetParams().LocalUserNum;
+	AccountInfo->PlatformUserId = InAsyncOp.GetParams().PlatformUserId;
 	AccountInfo->UserId = LocalUserId;
 	AccountInfo->LoginStatus = ELoginStatus::LoggedIn;
 
+	// Get display name
+	if (EOS_HUserInfo UserInfoHandle = EOS_Platform_GetUserInfoInterface(static_cast<FOnlineServicesEOS&>(GetServices()).GetEOSPlatformHandle()))
+	{
+		EOS_UserInfo_CopyUserInfoOptions Options = { };
+		Options.ApiVersion = EOS_USERINFO_COPYUSERINFO_API_LATEST;
+		Options.LocalUserId = EpicAccountId;
+		Options.TargetUserId = EpicAccountId;
+
+		EOS_UserInfo* UserInfo = nullptr;
+
+		EOS_EResult CopyResult = EOS_UserInfo_CopyUserInfo(UserInfoHandle, &Options, &UserInfo);
+		if (CopyResult == EOS_EResult::EOS_Success)
+		{
+			AccountInfo->DisplayName = UTF8_TO_TCHAR(UserInfo->DisplayName);
+			EOS_UserInfo_Release(UserInfo);
+		}
+	}
+
 	check(!AccountInfos.Contains(LocalUserId));
 	AccountInfos.Emplace(LocalUserId, AccountInfo);
 
@@ -592,17 +628,17 @@ TOnlineAsyncOpHandle<FAuthGenerateAuth> FAuthEOS::GenerateAuth(FAuthGenerateAuth
 	return AsyncOperation->GetHandle();
 }
 
-TOnlineResult<FAuthGetAccountByLocalUserNum> FAuthEOS::GetAccountByLocalUserNum(FAuthGetAccountByLocalUserNum::Params&& Params)
+TOnlineResult<FAuthGetAccountByPlatformUserId> FAuthEOS::GetAccountByPlatformUserId(FAuthGetAccountByPlatformUserId::Params&& Params)
 {
-	TResult<FOnlineAccountIdHandle, FOnlineError> LocalUserIdResult = GetAccountIdByLocalUserNum(Params.LocalUserNum);
-	if (LocalUserIdResult.IsOk())
+	TResult<FOnlineAccountIdHandle, FOnlineError> PlatformUserIdResult = GetAccountIdByPlatformUserId(Params.PlatformUserId);
+	if (PlatformUserIdResult.IsOk())
 	{
-		FAuthGetAccountByLocalUserNum::Result Result = { AccountInfos.FindChecked(LocalUserIdResult.GetOkValue()) };
-		return TOnlineResult<FAuthGetAccountByLocalUserNum>(Result);
+		FAuthGetAccountByPlatformUserId::Result Result = { AccountInfos.FindChecked(PlatformUserIdResult.GetOkValue()) };
+		return TOnlineResult<FAuthGetAccountByPlatformUserId>(Result);
 	}
 	else
 	{
-		return TOnlineResult<FAuthGetAccountByLocalUserNum>(LocalUserIdResult.GetErrorValue());
+		return TOnlineResult<FAuthGetAccountByPlatformUserId>(PlatformUserIdResult.GetErrorValue());
 	}
 }
 
@@ -627,11 +663,11 @@ bool FAuthEOS::IsLoggedIn(const FOnlineAccountIdHandle& AccountId) const
 	return AccountInfos.Contains(AccountId);
 }
 
-TResult<FOnlineAccountIdHandle, FOnlineError> FAuthEOS::GetAccountIdByLocalUserNum(int32 LocalUserNum) const
+TResult<FOnlineAccountIdHandle, FOnlineError> FAuthEOS::GetAccountIdByPlatformUserId(FPlatformUserId PlatformUserId) const
 {
 	for (const TPair<FOnlineAccountIdHandle, TSharedRef<FAccountInfoEOS>>& AccountPair : AccountInfos)
 	{
-		if (AccountPair.Value->LocalUserNum == LocalUserNum)
+		if (AccountPair.Value->PlatformUserId == PlatformUserId)
 		{
 			TResult<FOnlineAccountIdHandle, FOnlineError> Result(AccountPair.Key);
 			return Result;
diff --git a/Engine/Plugins/Online/Experimental/OnlineServicesEOS/Source/Private/ExternalUIEOS.cpp b/Engine/Plugins/Online/Experimental/OnlineServicesEOS/Source/Private/ExternalUIEOS.cpp
--- a/Engine/Plugins/Online/Experimental/OnlineServicesEOS/Source/Private/ExternalUIEOS.cpp
+++ b/Engine/Plugins/Online/Experimental/OnlineServicesEOS/Source/Private/ExternalUIEOS.cpp
@@ -23,14 +23,60 @@ FExternalUIEOS::FExternalUIEOS(FOnlineServicesEOS& InServices)
 
 void FExternalUIEOS::Initialize()
 {
-	FExternalUICommon::Initialize();
+	Super::Initialize();
 
 	UIHandle = EOS_Platform_GetUIInterface(static_cast<FOnlineServicesEOS&>(GetServices()).GetEOSPlatformHandle());
 	check(UIHandle != nullptr);
 }
 
+void FExternalUIEOS::LoadConfig()
+{
+	Super::LoadConfig();
+	LoadConfig(Config);
+}
+
 void FExternalUIEOS::PreShutdown()
 {
+	Super::PreShutdown();
+}
+
+TOnlineAsyncOpHandle<FExternalUIShowLoginUI> FExternalUIEOS::ShowLoginUI(FExternalUIShowLoginUI::Params&& Params)
+{
+	TOnlineAsyncOpRef<FExternalUIShowLoginUI> Op = GetOp<FExternalUIShowLoginUI>(MoveTemp(Params));
+	if (Config.bShowLoginUIEnabled)
+	{
+		// Use account portal auth login
+		Op->Then([this](TOnlineAsyncOp<FExternalUIShowLoginUI>& InAsyncOp)
+		{
+			FAuthLogin::Params AuthLoginParams;
+			AuthLoginParams.PlatformUserId = InAsyncOp.GetParams().PlatformUserId;
+			AuthLoginParams.CredentialsType = TEXT("accountportal");
+			AuthLoginParams.Scopes = InAsyncOp.GetParams().Scopes;
+
+			TPromise<void>* TempPromise = new TPromise<void>;
+			GetServices().GetAuthInterface()->Login(MoveTemp(AuthLoginParams))
+				.OnComplete([Op = InAsyncOp.AsShared(), TempPromise](const TOnlineResult<FAuthLogin>& AuthLoginResult) mutable
+				{
+					if (AuthLoginResult.IsOk())
+					{
+						FExternalUIShowLoginUI::Result ExtUIShowLoginResult = { AuthLoginResult.GetOkValue().AccountInfo };
+						Op->SetResult(CopyTemp(ExtUIShowLoginResult));
+					}
+					else
+					{
+						Op->SetError(CopyTemp(AuthLoginResult.GetErrorValue()));
+					}
+					TempPromise->SetValue();
+					delete TempPromise;
+				});
+			return TempPromise->GetFuture();
+		}).Enqueue(GetSerialQueue());
+	}
+	else
+	{
+		Op->SetError(Errors::NotImplemented());
+	}
+	return Op->GetHandle();
 }
 
 TOnlineAsyncOpHandle<FExternalUIShowFriendsUI> FExternalUIEOS::ShowFriendsUI(FExternalUIShowFriendsUI::Params&& Params)
diff --git a/Engine/Plugins/Online/Experimental/OnlineServicesEOS/Source/Public/AuthEOS.h b/Engine/Plugins/Online/Experimental/OnlineServicesEOS/Source/Public/AuthEOS.h
--- a/Engine/Plugins/Online/Experimental/OnlineServicesEOS/Source/Public/AuthEOS.h
+++ b/Engine/Plugins/Online/Experimental/OnlineServicesEOS/Source/Public/AuthEOS.h
@@ -26,7 +26,7 @@ public:
 	virtual TOnlineAsyncOpHandle<FAuthLogin> Login(FAuthLogin::Params&& Params) override;
 	virtual TOnlineAsyncOpHandle<FAuthLogout> Logout(FAuthLogout::Params&& Params) override;
 	virtual TOnlineAsyncOpHandle<FAuthGenerateAuth> GenerateAuth(FAuthGenerateAuth::Params&& Params) override;
-	virtual TOnlineResult<FAuthGetAccountByLocalUserNum> GetAccountByLocalUserNum(FAuthGetAccountByLocalUserNum::Params&& Params) override;
+	virtual TOnlineResult<FAuthGetAccountByPlatformUserId> GetAccountByPlatformUserId(FAuthGetAccountByPlatformUserId::Params&& Params) override;
 	virtual TOnlineResult<FAuthGetAccountByAccountId> GetAccountByAccountId(FAuthGetAccountByAccountId::Params&& Params) override;
 
 	bool IsLoggedIn(const FOnlineAccountIdHandle& AccountId) const;
@@ -43,7 +43,7 @@ public:
 
 protected:
 	void OnEOSLoginStatusChanged(FOnlineAccountIdHandle LocalUserId, ELoginStatus PreviousStatus, ELoginStatus CurrentStatus);
-	TResult<FOnlineAccountIdHandle, FOnlineError> GetAccountIdByLocalUserNum(int32 LocalUserNum) const;
+	TResult<FOnlineAccountIdHandle, FOnlineError> GetAccountIdByPlatformUserId(FPlatformUserId PlatformUserId) const;
 
 	void ProcessSuccessfulLogin(TOnlineAsyncOp<FAuthLogin>& InAsyncOp);
 
diff --git a/Engine/Plugins/Online/Experimental/OnlineServicesEOS/Source/Public/ExternalUIEOS.h b/Engine/Plugins/Online/Experimental/OnlineServicesEOS/Source/Public/ExternalUIEOS.h
--- a/Engine/Plugins/Online/Experimental/OnlineServicesEOS/Source/Public/ExternalUIEOS.h
+++ b/Engine/Plugins/Online/Experimental/OnlineServicesEOS/Source/Public/ExternalUIEOS.h
@@ -12,6 +12,12 @@
 
 namespace UE::Online {
 
+struct FExternalUIEOSConfig
+{
+	/** Is ShowLoginUI enabled? */
+	bool bShowLoginUIEnabled = true;
+};
+
 class FOnlineServicesEOS;
 
 class ONLINESERVICESEOS_API FExternalUIEOS : public FExternalUICommon
@@ -21,11 +27,25 @@ public:
 
 	FExternalUIEOS(FOnlineServicesEOS& InOwningSubsystem);
 	virtual void Initialize() override;
+	virtual void LoadConfig() override;
 	virtual void PreShutdown() override;
+	virtual TOnlineAsyncOpHandle<FExternalUIShowLoginUI> ShowLoginUI(FExternalUIShowLoginUI::Params&& Params) override;
 	virtual TOnlineAsyncOpHandle<FExternalUIShowFriendsUI> ShowFriendsUI(FExternalUIShowFriendsUI::Params&& Params) override;
 
 protected:
+	using Super::LoadConfig;
+	/** Handle to EOS UI */
 	EOS_HUI UIHandle;
+	/** Config */
+	FExternalUIEOSConfig Config;
 };
 
+namespace Meta {
+
+BEGIN_ONLINE_STRUCT_META(FExternalUIEOSConfig)
+	ONLINE_STRUCT_FIELD(FExternalUIEOSConfig, bShowLoginUIEnabled)
+END_ONLINE_STRUCT_META()
+
+/* Meta */ }
+
 /* UE::Online */ }
