Expansion of logic handling external auth types, allowing for specification via function parameters and default configuration per platform.

diff --git a/Engine/Config/Android/BaseAndroidEngine.ini b/Engine/Config/Android/BaseAndroidEngine.ini
--- a/Engine/Config/Android/BaseAndroidEngine.ini
+++ b/Engine/Config/Android/BaseAndroidEngine.ini
@@ -5,6 +5,9 @@ LocalPlatformName=AND
 [Audio]
 AudioCaptureModuleName=AudioCaptureAndroid
 
+[OnlineServices.EOS.Auth]
+DefaultExternalCredentialTypeStr=Google
+
 ;Uncomment to turn on AndroidBackgroundHttp through the AndroidFetchBackgroundDownload and AndroidBackgroundService modules
 ;[BackgroundHttp]
 ;PlatformModularFeatureName=AndroidFetchBackgroundDownload
\ No newline at end of file
diff --git a/Engine/Config/IOS/BaseIOSEngine.ini b/Engine/Config/IOS/BaseIOSEngine.ini
--- a/Engine/Config/IOS/BaseIOSEngine.ini
+++ b/Engine/Config/IOS/BaseIOSEngine.ini
@@ -6,4 +6,7 @@ LocalPlatformName=IOS
 AudioCaptureModuleName=AudioCaptureAudioUnit
 
 [Ping]
-StackSize=0
\ No newline at end of file
+StackSize=0
+
+[OnlineServices.EOS.Auth]
+DefaultExternalCredentialTypeStr=Apple
\ No newline at end of file
diff --git a/Engine/Config/Linux/BaseLinuxEngine.ini b/Engine/Config/Linux/BaseLinuxEngine.ini
--- a/Engine/Config/Linux/BaseLinuxEngine.ini
+++ b/Engine/Config/Linux/BaseLinuxEngine.ini
@@ -1,2 +1,5 @@
 [PlatformPaths]
-UnrealBuildTool=Engine/Binaries/DotNET/UnrealBuildTool/UnrealBuildTool
\ No newline at end of file
+UnrealBuildTool=Engine/Binaries/DotNET/UnrealBuildTool/UnrealBuildTool
+
+[OnlineServices.EOS.Auth]
+DefaultExternalCredentialTypeStr=Steam
\ No newline at end of file
diff --git a/Engine/Config/Mac/BaseMacEngine.ini b/Engine/Config/Mac/BaseMacEngine.ini
--- a/Engine/Config/Mac/BaseMacEngine.ini
+++ b/Engine/Config/Mac/BaseMacEngine.ini
@@ -5,4 +5,7 @@ AudioCaptureModuleName=AudioCaptureRtAudio
 LocalPlatformName=MAC
 
 [PlatformPaths]
-UnrealBuildTool=Engine/Binaries/DotNET/UnrealBuildTool/UnrealBuildTool
\ No newline at end of file
+UnrealBuildTool=Engine/Binaries/DotNET/UnrealBuildTool/UnrealBuildTool
+
+[OnlineServices.EOS.Auth]
+DefaultExternalCredentialTypeStr=Steam
\ No newline at end of file
diff --git a/Engine/Config/Windows/BaseWindowsEngine.ini b/Engine/Config/Windows/BaseWindowsEngine.ini
--- a/Engine/Config/Windows/BaseWindowsEngine.ini
+++ b/Engine/Config/Windows/BaseWindowsEngine.ini
@@ -6,3 +6,6 @@ LocalPlatformName=WIN
 
 [DeviceProfileManager]
 DeviceProfileSelectionModule="WindowsDeviceProfileSelector"
+
+[OnlineServices.EOS.Auth]
+DefaultExternalCredentialTypeStr=Steam
\ No newline at end of file
diff --git a/Engine/Plugins/Online/Experimental/OnlineServicesEOS/Source/Private/AuthEOS.cpp b/Engine/Plugins/Online/Experimental/OnlineServicesEOS/Source/Private/AuthEOS.cpp
--- a/Engine/Plugins/Online/Experimental/OnlineServicesEOS/Source/Private/AuthEOS.cpp
+++ b/Engine/Plugins/Online/Experimental/OnlineServicesEOS/Source/Private/AuthEOS.cpp
@@ -22,6 +22,19 @@
 
 namespace UE::Online {
 
+struct FAuthEOSConfig
+{
+	FString DefaultExternalCredentialTypeStr;
+};
+
+namespace Meta {
+
+BEGIN_ONLINE_STRUCT_META(FAuthEOSConfig)
+	ONLINE_STRUCT_FIELD(FAuthEOSConfig, DefaultExternalCredentialTypeStr)
+END_ONLINE_STRUCT_META()
+
+/* Meta*/ }
+
 static inline ELoginStatus ToELoginStatus(EOS_ELoginStatus InStatus)
 {
 	switch (InStatus)
@@ -241,10 +254,6 @@ bool LexFromString(EOS_ELoginCredentialType& OutEnum, const TCHAR* const InStrin
 	{
 		OutEnum = EOS_ELoginCredentialType::EOS_LCT_PersistentAuth;
 	}
-	//else if (FCString::Stricmp(InString, TEXT("DeviceCode")) == 0) // DeviceCode is deprecated
-	//{
-	//	OutEnum = EOS_ELoginCredentialType::EOS_LCT_DeviceCode;
-	//}
 	else if (FCString::Stricmp(InString, TEXT("Password")) == 0)
 	{
 		OutEnum = EOS_ELoginCredentialType::EOS_LCT_Password;
@@ -272,6 +281,45 @@ bool LexFromString(EOS_ELoginCredentialType& OutEnum, const TCHAR* const InStrin
 	return true;
 }
 
+bool LexFromString(EOS_EExternalCredentialType& OutEnum, const TCHAR* InString)
+{
+	if (FCString::Stricmp(InString, TEXT("Steam")) == 0)
+	{
+		OutEnum = EOS_EExternalCredentialType::EOS_ECT_STEAM_APP_TICKET;
+	}
+	else if (FCString::Stricmp(InString, TEXT("PSN")) == 0)
+	{
+		OutEnum = EOS_EExternalCredentialType::EOS_ECT_PSN_ID_TOKEN;
+	}
+	else if (FCString::Stricmp(InString, TEXT("Xbox")) == 0)
+	{
+		OutEnum = EOS_EExternalCredentialType::EOS_ECT_XBL_XSTS_TOKEN;
+	}
+	else if (FCString::Stricmp(InString, TEXT("Nintendo")) == 0)
+	{
+		OutEnum = EOS_EExternalCredentialType::EOS_ECT_NINTENDO_ID_TOKEN;
+	}
+	else if (FCString::Stricmp(InString, TEXT("NSA")) == 0)
+	{
+		OutEnum = EOS_EExternalCredentialType::EOS_ECT_NINTENDO_NSA_ID_TOKEN;
+	}
+	else if (FCString::Stricmp(InString, TEXT("Apple")) == 0)
+	{
+		OutEnum = EOS_EExternalCredentialType::EOS_ECT_APPLE_ID_TOKEN;
+	}
+	else if (FCString::Stricmp(InString, TEXT("Google")) == 0)
+	{
+		OutEnum = EOS_EExternalCredentialType::EOS_ECT_GOOGLE_ID_TOKEN;
+	}
+	else
+	{
+		// Unknown means OpenID
+		OutEnum = EOS_EExternalCredentialType::EOS_ECT_OPENID_ACCESS_TOKEN;
+	}
+
+	return true;
+}
+
 TOnlineAsyncOpHandle<FAuthLogin> FAuthEOS::Login(FAuthLogin::Params&& Params)
 {
 	// Is this auto-login?
@@ -319,13 +367,37 @@ TOnlineAsyncOpHandle<FAuthLogin> FAuthEOS::Login(FAuthLogin::Params&& Params)
 		LoginOptions.ScopeFlags = EOS_EAuthScopeFlags::EOS_AS_BasicProfile | EOS_EAuthScopeFlags::EOS_AS_FriendsList | EOS_EAuthScopeFlags::EOS_AS_Presence;
 	}	
 
+	// First, we'll check if any subtypes have been specified
+	FString CredentialTypeStr;
+	FString ExternalCredentialTypeStr;
+	Op->GetParams().CredentialsType.Split(FString(TEXT(":")), &CredentialTypeStr, &ExternalCredentialTypeStr);
+
+	// If no subtype was specified, we'll treat the string passed as the type
+	if (CredentialTypeStr.IsEmpty())
+	{
+		CredentialTypeStr = Op->GetParams().CredentialsType;
+	}
+
 	FEOSAuthCredentials Credentials;
-	if (LexFromString(Credentials.Type, *Op->GetParams().CredentialsType))
+	if (LexFromString(Credentials.Type, *CredentialTypeStr))
 	{
 		switch (Credentials.Type)
 		{
 		case EOS_ELoginCredentialType::EOS_LCT_ExternalAuth:
-			Credentials.SetExternalCredentialType(EOnlineServices::Steam, Op->GetParams().CredentialsId); // TODO: Will be GetServicesProvider() or something similar
+			// If an external credential type wasn't specified, we'll grab the platform default from the configuration values
+			if (ExternalCredentialTypeStr.IsEmpty())
+			{
+				FAuthEOSConfig AuthEOSConfig;
+				LoadConfig(AuthEOSConfig);
+				ExternalCredentialTypeStr = AuthEOSConfig.DefaultExternalCredentialTypeStr;
+			}			
+			
+			EOS_EExternalCredentialType ExternalType;
+			if (LexFromString(ExternalType, *ExternalCredentialTypeStr))
+			{
+				Credentials.ExternalType = ExternalType;
+			}
+
 			Credentials.SetToken(Op->GetParams().CredentialsToken);
 			break;
 		case EOS_ELoginCredentialType::EOS_LCT_ExchangeCode:
