Initial checkin of OSSv2 lobbies and EOS integration

diff --git a/Engine/Plugins/Online/Experimental/OnlineServices/Source/OnlineServicesCommon/Private/Online/LobbiesCommon.cpp b/Engine/Plugins/Online/Experimental/OnlineServices/Source/OnlineServicesCommon/Private/Online/LobbiesCommon.cpp
new file mode 100644
--- /dev/null
+++ b/Engine/Plugins/Online/Experimental/OnlineServices/Source/OnlineServicesCommon/Private/Online/LobbiesCommon.cpp
@@ -0,0 +1,1019 @@
+// Copyright Epic Games, Inc. All Rights Reserved.
+
+#include "Online/LobbiesCommon.h"
+
+#include "Online/Auth.h"
+#include "Online/OnlineServicesCommon.h"
+
+DEFINE_LOG_CATEGORY(LogLobbies);
+
+namespace UE::Online {
+
+template <typename DataType, typename OpType>
+const DataType& GetOpDataChecked(const TOnlineAsyncOp<OpType>& Op, const FString& Key)
+{
+	const DataType* Data = Op.Data.template Get<DataType>(Key);
+	check(Data);
+	return *Data;
+}
+
+TSharedPtr<FLobbySchema> FLobbySchema::Create(FLobbySchemaConfig LobbySchemaConfig)
+{
+	TSharedRef<FLobbySchema> Schema = MakeShared<FLobbySchema>();
+	Schema->SchemaId = LobbySchemaConfig.SchemaName;
+
+	// Todo: for real validation and init
+	if (Schema->SchemaId == FLobbySchemaId())
+	{
+		return TSharedPtr<FLobbySchema>();
+	}
+	else
+	{
+		return Schema;
+	}
+}
+
+TSharedRef<FLobbyServiceAttributeChanges> FLobbySchema::TranslateLobbyAttributes(const TMap<FLobbyAttributeId, FLobbyVariant>& LobbyAttributes, const FLobbyClientAttributeChanges& ClientAttributeChanges) const
+{
+	// Todo: implementation.
+
+	TSharedRef<FLobbyServiceAttributeChanges> Changes = MakeShared<FLobbyServiceAttributeChanges>();
+	Changes->MutatedAttributes = ClientAttributeChanges.MutatedAttributes;
+	Changes->ClearedAttributes = ClientAttributeChanges.ClearedAttributes;
+	return Changes;
+}
+
+TSharedRef<FLobbyClientAttributeChanges> FLobbySchema::TranslateLobbyAttributes(const FLobbyServiceAttributeChanges& ServiceAttributeChanges) const
+{
+	// Todo: implementation.
+
+	TSharedRef<FLobbyClientAttributeChanges> Changes = MakeShared<FLobbyClientAttributeChanges>();
+	Changes->MutatedAttributes = ServiceAttributeChanges.MutatedAttributes;
+	Changes->ClearedAttributes = ServiceAttributeChanges.ClearedAttributes;
+	return Changes;
+}
+
+TSharedRef<FLobbyServiceAttributeChanges> FLobbySchema::TranslateLobbyMemberAttributes(const TMap<FLobbyAttributeId, FLobbyVariant>& LobbyAttributes, const FLobbyClientAttributeChanges& ClientAttributeChanges) const
+{
+	// Todo: implementation.
+
+	TSharedRef<FLobbyServiceAttributeChanges> Changes = MakeShared<FLobbyServiceAttributeChanges>();
+	Changes->MutatedAttributes = ClientAttributeChanges.MutatedAttributes;
+	Changes->ClearedAttributes = ClientAttributeChanges.ClearedAttributes;
+	return Changes;
+}
+
+TSharedRef<FLobbyClientAttributeChanges> FLobbySchema::TranslateLobbyMemberAttributes(const FLobbyServiceAttributeChanges& ServiceAttributeChanges) const
+{
+	// Todo: implementation.
+
+	TSharedRef<FLobbyClientAttributeChanges> Changes = MakeShared<FLobbyClientAttributeChanges>();
+	Changes->MutatedAttributes = ServiceAttributeChanges.MutatedAttributes;
+	Changes->ClearedAttributes = ServiceAttributeChanges.ClearedAttributes;
+	return Changes;
+}
+
+bool FLobbySchemaRegistry::Initialize(TArray<FLobbySchemaConfig> LobbySchemaConfigs)
+{
+	bool bSuccess = true;
+
+	for (FLobbySchemaConfig& LobbySchemaConfig : LobbySchemaConfigs)
+	{
+		bSuccess &= RegisterSchema(MoveTemp(LobbySchemaConfig));
+	}
+
+	return bSuccess;
+}
+
+TSharedPtr<FLobbySchema> FLobbySchemaRegistry::FindSchema(FLobbySchemaId SchemaId)
+{
+	TSharedRef<FLobbySchema>* Schema = RegisteredSchemas.Find(SchemaId);
+	return Schema ? *Schema : TSharedPtr<FLobbySchema>();
+}
+
+bool FLobbySchemaRegistry::RegisterSchema(FLobbySchemaConfig LobbySchemaConfig)
+{
+	// Todo: handle schema hierarchy.
+	// Todo: handle schema parsing.
+
+	FLobbySchemaId SchemaName = LobbySchemaConfig.SchemaName;
+
+	if (TSharedPtr<FLobbySchema> Schema = FLobbySchema::Create(MoveTemp(LobbySchemaConfig)))
+	{
+		RegisteredSchemas.Emplace(SchemaName, Schema.ToSharedRef());
+		return true;
+	}
+	else
+	{
+		return false;
+	}
+}
+
+FLobbiesCommon::FLobbiesCommon(FOnlineServicesCommon& InServices)
+	: TOnlineComponent(TEXT("Lobbies"), InServices)
+{
+}
+
+void FLobbiesCommon::Initialize()
+{
+	TOnlineComponent<ILobbies>::Initialize();
+
+	LobbySchemaRegistry = MakeShared<FLobbySchemaRegistry>();
+	ServiceSchema = MakeShared<FLobbySchema>();
+
+	FLobbyConfig LobbyConfig;
+	// Todo: LoadConfig(LobbyConfig);
+
+	if (LobbyConfig.RegisteredSchemas.Num() > 0)
+	{
+		TArray<FLobbySchemaConfig> LobbySchemaConfigs;
+		LobbySchemaConfigs.Reserve(LobbyConfig.RegisteredSchemas.Num());
+
+		for (FLobbySchemaId LobbySchemaId : LobbyConfig.RegisteredSchemas)
+		{
+			FLobbySchemaConfig LobbySchemaConfig;
+			// Todo: LoadConfig(LobbySchemaConfig, LobbySchemaId.ToString());
+			LobbySchemaConfigs.Add(MoveTemp(LobbySchemaConfig));
+		}
+
+		if (!LobbySchemaRegistry->Initialize(MoveTemp(LobbySchemaConfigs)))
+		{
+			// Todo: Log error.
+		}
+	}
+	else
+	{
+		// Todo: Log error.
+	}
+}
+
+void FLobbiesCommon::RegisterCommands()
+{
+	TOnlineComponent<ILobbies>::RegisterCommands();
+
+	RegisterCommand(&FLobbiesCommon::CreateLobby);
+	RegisterCommand(&FLobbiesCommon::FindLobby);
+	RegisterCommand(&FLobbiesCommon::RestoreLobbies);
+	RegisterCommand(&FLobbiesCommon::JoinLobby);
+	RegisterCommand(&FLobbiesCommon::LeaveLobby);
+	RegisterCommand(&FLobbiesCommon::InviteLobbyMember);
+	RegisterCommand(&FLobbiesCommon::DeclineLobbyInvitation);
+	RegisterCommand(&FLobbiesCommon::KickLobbyMember);
+	RegisterCommand(&FLobbiesCommon::PromoteLobbyMember);
+	RegisterCommand(&FLobbiesCommon::ModifyLobbySchema);
+	RegisterCommand(&FLobbiesCommon::ModifyLobbyJoinPolicy);
+	RegisterCommand(&FLobbiesCommon::ModifyLobbyAttributes);
+	RegisterCommand(&FLobbiesCommon::ModifyLobbyMemberAttributes);
+	RegisterCommand(&FLobbiesCommon::GetJoinedLobbies);
+	RegisterCommand(&FLobbiesCommon::GetReceivedInvitations);
+
+	RegisterCommand(&FLobbiesCommon::FunctionalTest);
+}
+
+TOnlineAsyncOpHandle<FCreateLobby> FLobbiesCommon::CreateLobby(FCreateLobby::Params&& Params)
+{
+	TOnlineAsyncOpRef<FCreateLobby> Operation = GetOp<FCreateLobby>(MoveTemp(Params));
+	Operation->SetError(Errors::NotImplemented());
+	return Operation->GetHandle();
+}
+
+TOnlineAsyncOpHandle<FFindLobby> FLobbiesCommon::FindLobby(FFindLobby::Params&& Params)
+{
+	TOnlineAsyncOpRef<FFindLobby> Operation = GetOp<FFindLobby>(MoveTemp(Params));
+	Operation->SetError(Errors::NotImplemented());
+	return Operation->GetHandle();
+}
+
+TOnlineAsyncOpHandle<FRestoreLobbies> FLobbiesCommon::RestoreLobbies(FRestoreLobbies::Params&& Params)
+{
+	TOnlineAsyncOpRef<FRestoreLobbies> Operation = GetOp<FRestoreLobbies>(MoveTemp(Params));
+	Operation->SetError(Errors::NotImplemented());
+	return Operation->GetHandle();
+}
+
+TOnlineAsyncOpHandle<FJoinLobby> FLobbiesCommon::JoinLobby(FJoinLobby::Params&& Params)
+{
+	TOnlineAsyncOpRef<FJoinLobby> Operation = GetOp<FJoinLobby>(MoveTemp(Params));
+	Operation->SetError(Errors::NotImplemented());
+	return Operation->GetHandle();
+}
+
+TOnlineAsyncOpHandle<FLeaveLobby> FLobbiesCommon::LeaveLobby(FLeaveLobby::Params&& Params)
+{
+	TOnlineAsyncOpRef<FLeaveLobby> Operation = GetOp<FLeaveLobby>(MoveTemp(Params));
+	Operation->SetError(Errors::NotImplemented());
+	return Operation->GetHandle();
+}
+
+TOnlineAsyncOpHandle<FInviteLobbyMember> FLobbiesCommon::InviteLobbyMember(FInviteLobbyMember::Params&& Params)
+{
+	TOnlineAsyncOpRef<FInviteLobbyMember> Operation = GetOp<FInviteLobbyMember>(MoveTemp(Params));
+	Operation->SetError(Errors::NotImplemented());
+	return Operation->GetHandle();
+}
+
+TOnlineAsyncOpHandle<FDeclineLobbyInvitation> FLobbiesCommon::DeclineLobbyInvitation(FDeclineLobbyInvitation::Params&& Params)
+{
+	TOnlineAsyncOpRef<FDeclineLobbyInvitation> Operation = GetOp<FDeclineLobbyInvitation>(MoveTemp(Params));
+	Operation->SetError(Errors::NotImplemented());
+	return Operation->GetHandle();
+}
+
+TOnlineAsyncOpHandle<FKickLobbyMember> FLobbiesCommon::KickLobbyMember(FKickLobbyMember::Params&& Params)
+{
+	TOnlineAsyncOpRef<FKickLobbyMember> Operation = GetOp<FKickLobbyMember>(MoveTemp(Params));
+	Operation->SetError(Errors::NotImplemented());
+	return Operation->GetHandle();
+}
+
+TOnlineAsyncOpHandle<FPromoteLobbyMember> FLobbiesCommon::PromoteLobbyMember(FPromoteLobbyMember::Params&& Params)
+{
+	TOnlineAsyncOpRef<FPromoteLobbyMember> Operation = GetOp<FPromoteLobbyMember>(MoveTemp(Params));
+	Operation->SetError(Errors::NotImplemented());
+	return Operation->GetHandle();
+}
+
+TOnlineAsyncOpHandle<FModifyLobbySchema> FLobbiesCommon::ModifyLobbySchema(FModifyLobbySchema::Params&& Params)
+{
+	TOnlineAsyncOpRef<FModifyLobbySchema> Operation = GetOp<FModifyLobbySchema>(MoveTemp(Params));
+	Operation->SetError(Errors::NotImplemented());
+	return Operation->GetHandle();
+}
+
+TOnlineAsyncOpHandle<FModifyLobbyJoinPolicy> FLobbiesCommon::ModifyLobbyJoinPolicy(FModifyLobbyJoinPolicy::Params&& Params)
+{
+	TOnlineAsyncOpRef<FModifyLobbyJoinPolicy> Operation = GetOp<FModifyLobbyJoinPolicy>(MoveTemp(Params));
+	Operation->SetError(Errors::NotImplemented());
+	return Operation->GetHandle();
+}
+
+TOnlineAsyncOpHandle<FModifyLobbyAttributes> FLobbiesCommon::ModifyLobbyAttributes(FModifyLobbyAttributes::Params&& Params)
+{
+	TOnlineAsyncOpRef<FModifyLobbyAttributes> Operation = GetOp<FModifyLobbyAttributes>(MoveTemp(Params));
+	Operation->SetError(Errors::NotImplemented());
+	return Operation->GetHandle();
+}
+
+TOnlineAsyncOpHandle<FModifyLobbyMemberAttributes> FLobbiesCommon::ModifyLobbyMemberAttributes(FModifyLobbyMemberAttributes::Params&& Params)
+{
+	TOnlineAsyncOpRef<FModifyLobbyMemberAttributes> Operation = GetOp<FModifyLobbyMemberAttributes>(MoveTemp(Params));
+	Operation->SetError(Errors::NotImplemented());
+	return Operation->GetHandle();
+}
+
+TOnlineResult<FGetJoinedLobbies> FLobbiesCommon::GetJoinedLobbies(FGetJoinedLobbies::Params&& Params)
+{
+	return TOnlineResult<FGetJoinedLobbies>(Errors::NotImplemented());
+}
+
+TOnlineResult<FGetReceivedInvitations> FLobbiesCommon::GetReceivedInvitations(FGetReceivedInvitations::Params&& Params)
+{
+	return TOnlineResult<FGetReceivedInvitations>(Errors::NotImplemented());
+}
+
+TOnlineEvent<void(const FLobbyJoined&)> FLobbiesCommon::OnLobbyJoined()
+{
+	return OnLobbyJoinedEvent;
+}
+
+TOnlineEvent<void(const FLobbyLeft&)> FLobbiesCommon::OnLobbyLeft()
+{
+	return OnLobbyLeftEvent;
+}
+
+TOnlineEvent<void(const FLobbyMemberJoined&)> FLobbiesCommon::OnLobbyMemberJoined()
+{
+	return OnLobbyMemberJoinedEvent;
+}
+
+TOnlineEvent<void(const FLobbyMemberLeft&)> FLobbiesCommon::OnLobbyMemberLeft()
+{
+	return OnLobbyMemberLeftEvent;
+}
+
+TOnlineEvent<void(const FLobbyLeaderChanged&)> FLobbiesCommon::OnLobbyLeaderChanged()
+{
+	return OnLobbyLeaderChangedEvent;
+}
+
+TOnlineEvent<void(const FLobbySchemaChanged&)> FLobbiesCommon::OnLobbySchemaChanged()
+{
+	return OnLobbySchemaChangedEvent;
+}
+
+TOnlineEvent<void(const FLobbyAttributesChanged&)> FLobbiesCommon::OnLobbyAttributesChanged()
+{
+	return OnLobbyAttributesChangedEvent;
+}
+
+TOnlineEvent<void(const FLobbyMemberAttributesChanged&)> FLobbiesCommon::OnLobbyMemberAttributesChanged()
+{
+	return OnLobbyMemberAttributesChangedEvent;
+}
+
+TOnlineEvent<void(const FLobbyInvitationAdded&)> FLobbiesCommon::OnLobbyInvitationAdded()
+{
+	return OnLobbyInvitationAddedEvent;
+}
+
+TOnlineEvent<void(const FLobbyInvitationRemoved&)> FLobbiesCommon::OnLobbyInvitationRemoved()
+{
+	return OnLobbyInvitationRemovedEvent;
+}
+
+TFuture<TDefaultErrorResultInternal<FOnlineLobbyIdHandle>> FLobbiesCommon::AwaitInvitation(
+	FOnlineAccountIdHandle TargetAccountId,
+	FOnlineLobbyIdHandle LobbyId,
+	float TimeoutSeconds)
+{
+	// todo: Make this nicer - current implementation has issue on shutdown.
+	struct AwaitInvitationState
+	{
+		TPromise<TDefaultErrorResultInternal<FOnlineLobbyIdHandle>> Promise;
+		FOnlineEventDelegateHandle OnLobbyInvitationAddedHandle;
+		FTSTicker::FDelegateHandle OnAwaitExpiredHandle;
+		bool bSignaled = false;
+	};
+
+	TSharedRef<AwaitInvitationState> AwaitState = MakeShared<AwaitInvitationState>();
+
+	AwaitState->OnLobbyInvitationAddedHandle = OnLobbyInvitationAdded().Add(
+	[this, TargetAccountId, LobbyId, AwaitState](const FLobbyInvitationAdded& Notification)
+	{
+		FOnlineEventDelegateHandle DelegateHandle = MoveTemp(AwaitState->OnLobbyInvitationAddedHandle);
+
+		if (!AwaitState->bSignaled)
+		{
+			if (Notification.Lobby->LobbyId == LobbyId && Notification.LocalUserId == TargetAccountId)
+			{
+				FTSTicker::GetCoreTicker().RemoveTicker(AwaitState->OnAwaitExpiredHandle);
+				AwaitState->bSignaled = true;
+				AwaitState->Promise.EmplaceValue(LobbyId);
+			}
+		}
+	});
+
+	AwaitState->OnAwaitExpiredHandle = FTSTicker::GetCoreTicker().AddTicker(
+	TEXT("FLobbiesCommon::AwaitInvitation"),
+	TimeoutSeconds,
+	[this, AwaitState](float)
+	{
+		if (!AwaitState->bSignaled)
+		{
+			// Todo: Errors.
+			AwaitState->bSignaled = true;
+			AwaitState->OnLobbyInvitationAddedHandle.Unbind();
+			AwaitState->Promise.EmplaceValue(Errors::NotImplemented());
+			FTSTicker::GetCoreTicker().RemoveTicker(AwaitState->OnAwaitExpiredHandle);
+		}
+
+		return false;
+	});
+
+	return AwaitState->Promise.GetFuture();
+}
+
+#if LOBBIES_FUNCTIONAL_TEST_ENABLED
+struct FFunctionalTestConfig
+{
+	FString TestAccount1Type;
+	FString TestAccount1Id;
+	FString TestAccount1Token;
+
+	FString TestAccount2Type;
+	FString TestAccount2Id;
+	FString TestAccount2Token;
+
+	float InvitationWaitSeconds = 10.f;
+};
+
+namespace Meta {
+
+BEGIN_ONLINE_STRUCT_META(FFunctionalTestConfig)
+	ONLINE_STRUCT_FIELD(FFunctionalTestConfig, TestAccount1Type),
+	ONLINE_STRUCT_FIELD(FFunctionalTestConfig, TestAccount1Id),
+	ONLINE_STRUCT_FIELD(FFunctionalTestConfig, TestAccount1Token),
+
+	ONLINE_STRUCT_FIELD(FFunctionalTestConfig, TestAccount2Type),
+	ONLINE_STRUCT_FIELD(FFunctionalTestConfig, TestAccount2Id),
+	ONLINE_STRUCT_FIELD(FFunctionalTestConfig, TestAccount2Token)
+END_ONLINE_STRUCT_META()
+
+/* Meta */ }
+
+class FLobbyEventCapture final
+{
+public:
+	UE_NONCOPYABLE(FLobbyEventCapture);
+
+	FLobbyEventCapture(FLobbiesCommon& Lobbies)
+		: LobbyJoinedHandle(Lobbies.OnLobbyJoined().Add(this, &FLobbyEventCapture::OnLobbyJoined))
+		, LobbyLeftHandle(Lobbies.OnLobbyLeft().Add(this, &FLobbyEventCapture::OnLobbyLeft))
+		, LobbyMemberJoinedHandle(Lobbies.OnLobbyMemberJoined().Add(this, &FLobbyEventCapture::OnLobbyMemberJoined))
+		, LobbyMemberLeftHandle(Lobbies.OnLobbyMemberLeft().Add(this, &FLobbyEventCapture::OnLobbyMemberLeft))
+		, LobbyLeaderChangedHandle(Lobbies.OnLobbyLeaderChanged().Add(this, &FLobbyEventCapture::OnLobbyLeaderChanged))
+		, LobbySchemaChangedHandle(Lobbies.OnLobbySchemaChanged().Add(this, &FLobbyEventCapture::OnLobbySchemaChanged))
+		, LobbyAttributesChangedHandle(Lobbies.OnLobbyAttributesChanged().Add(this, &FLobbyEventCapture::OnLobbyAttributesChanged))
+		, LobbyMemberAttributesChangedHandle(Lobbies.OnLobbyMemberAttributesChanged().Add(this, &FLobbyEventCapture::OnLobbyMemberAttributesChanged))
+		, LobbyInvitationAddedHandle(Lobbies.OnLobbyInvitationAdded().Add(this, &FLobbyEventCapture::OnLobbyInvitationAdded))
+		, LobbyInvitationRemovedHandle(Lobbies.OnLobbyInvitationRemoved().Add(this, &FLobbyEventCapture::OnLobbyInvitationRemoved))
+	{
+	}
+
+	void Empty()
+	{
+		LobbyJoined.Empty();
+		LobbyLeft.Empty();
+		LobbyMemberJoined.Empty();
+		LobbyMemberLeft.Empty();
+		LobbyLeaderChanged.Empty();
+		LobbySchemaChanged.Empty();
+		LobbyAttributesChanged.Empty();
+		LobbyMemberAttributesChanged.Empty();
+		LobbyInvitationAdded.Empty();
+		LobbyInvitationRemoved.Empty();
+		NextIndex = 0;
+		TotalNotificationsReceived = 0;
+	}
+
+	uint32 GetTotalNotificationsReceived() const
+	{
+		return TotalNotificationsReceived;
+	}
+
+	template <typename NotificationType>
+	struct TNotificationInfo
+	{
+		NotificationType Notification;
+		int32 GlobalIndex = 0;
+	};
+
+	TArray<TNotificationInfo<FLobbyJoined>> LobbyJoined;
+	TArray<TNotificationInfo<FLobbyLeft>> LobbyLeft;
+	TArray<TNotificationInfo<FLobbyMemberJoined>> LobbyMemberJoined;
+	TArray<TNotificationInfo<FLobbyMemberLeft>> LobbyMemberLeft;
+	TArray<TNotificationInfo<FLobbyLeaderChanged>> LobbyLeaderChanged;
+	TArray<TNotificationInfo<FLobbySchemaChanged>> LobbySchemaChanged;
+	TArray<TNotificationInfo<FLobbyAttributesChanged>> LobbyAttributesChanged;
+	TArray<TNotificationInfo<FLobbyMemberAttributesChanged>> LobbyMemberAttributesChanged;
+	TArray<TNotificationInfo<FLobbyInvitationAdded>> LobbyInvitationAdded;
+	TArray<TNotificationInfo<FLobbyInvitationRemoved>> LobbyInvitationRemoved;
+
+private:
+	template <typename ContainerType, typename NotificationType>
+	void AddEvent(ContainerType& Container, const NotificationType& Notification)
+	{
+		Container.Add({Notification, NextIndex++});
+		++TotalNotificationsReceived;
+	}
+
+	void OnLobbyJoined(const FLobbyJoined& Notification) { AddEvent(LobbyJoined, Notification); }
+	void OnLobbyLeft(const FLobbyLeft& Notification) { AddEvent(LobbyLeft, Notification); }
+	void OnLobbyMemberJoined(const FLobbyMemberJoined& Notification) { AddEvent(LobbyMemberJoined, Notification); }
+	void OnLobbyMemberLeft(const FLobbyMemberLeft& Notification) { AddEvent(LobbyMemberLeft, Notification); }
+	void OnLobbyLeaderChanged(const FLobbyLeaderChanged& Notification) { AddEvent(LobbyLeaderChanged, Notification); }
+	void OnLobbySchemaChanged(const FLobbySchemaChanged& Notification) { AddEvent(LobbySchemaChanged, Notification); }
+	void OnLobbyAttributesChanged(const FLobbyAttributesChanged& Notification) { AddEvent(LobbyAttributesChanged, Notification); }
+	void OnLobbyMemberAttributesChanged(const FLobbyMemberAttributesChanged& Notification) { AddEvent(LobbyMemberAttributesChanged, Notification); }
+	void OnLobbyInvitationAdded(const FLobbyInvitationAdded& Notification) { AddEvent(LobbyInvitationAdded, Notification); }
+	void OnLobbyInvitationRemoved(const FLobbyInvitationRemoved& Notification) { AddEvent(LobbyInvitationRemoved, Notification); }
+
+	FOnlineEventDelegateHandle LobbyJoinedHandle;
+	FOnlineEventDelegateHandle LobbyLeftHandle;
+	FOnlineEventDelegateHandle LobbyMemberJoinedHandle;
+	FOnlineEventDelegateHandle LobbyMemberLeftHandle;
+	FOnlineEventDelegateHandle LobbyLeaderChangedHandle;
+	FOnlineEventDelegateHandle LobbySchemaChangedHandle;
+	FOnlineEventDelegateHandle LobbyAttributesChangedHandle;
+	FOnlineEventDelegateHandle LobbyMemberAttributesChangedHandle;
+	FOnlineEventDelegateHandle LobbyInvitationAddedHandle;
+	FOnlineEventDelegateHandle LobbyInvitationRemovedHandle;
+	int32 NextIndex = 0;
+	uint32 TotalNotificationsReceived = 0;
+};
+
+static const FString ConfigName = TEXT("FunctionalTest");
+static const FString User1KeyName = TEXT("User1");
+static const FString User2KeyName = TEXT("User2");
+static const FString CreateLobbyKeyName = TEXT("CreateLobby");
+static const FString FindLobbyKeyName = TEXT("FindLobby");
+static const FString LobbyEventCaptureKeyName = TEXT("LobbyEventCapture");
+static const FString ConfigNameKeyName = TEXT("Config");
+
+TOnlineAsyncOpHandle<FFunctionalTestLobbies> FLobbiesCommon::FunctionalTest(FFunctionalTestLobbies::Params&& InParams)
+{
+	TSharedRef<FFunctionalTestConfig> TestConfig = MakeShared<FFunctionalTestConfig>();
+	LoadConfig(*TestConfig, ConfigName);
+
+	TOnlineAsyncOpRef<FFunctionalTestLobbies> Op = GetOp<FFunctionalTestLobbies>(MoveTemp(InParams));
+
+	// Set up event capturing.
+	Op->Data.Set(LobbyEventCaptureKeyName, MakeShared<FLobbyEventCapture>(*this));
+
+	Op->Then(ConsumeStepResult<FFunctionalTestLogoutAllUsers>(*Op, &FLobbiesCommon::FunctionalTestLogoutAllUsers, FFunctionalTestLogoutAllUsers::Params{}))
+	.Then(CaptureStepResult<FFunctionalTestLoginUser>(*Op, User1KeyName, &FLobbiesCommon::FunctionalTestLoginUser, FFunctionalTestLoginUser::Params{FPlatformMisc::GetPlatformUserForUserIndex(0), TestConfig->TestAccount1Type, TestConfig->TestAccount1Id, TestConfig->TestAccount1Token}))
+	.Then(CaptureStepResult<FFunctionalTestLoginUser>(*Op, User2KeyName, &FLobbiesCommon::FunctionalTestLoginUser, FFunctionalTestLoginUser::Params{FPlatformMisc::GetPlatformUserForUserIndex(1), TestConfig->TestAccount2Type, TestConfig->TestAccount2Id, TestConfig->TestAccount2Token}))
+
+	//----------------------------------------------------------------------------------------------
+	// Test 1:
+	//    Step 1: Create a lobby with primary user.
+	//    Step 2: Leave lobby with primary user.
+	//----------------------------------------------------------------------------------------------
+
+	.Then([this](TOnlineAsyncOp<FFunctionalTestLobbies>& InAsyncOp)
+	{
+		// Prepare for next test check.
+		GetOpDataChecked<TSharedRef<FLobbyEventCapture>>(InAsyncOp, LobbyEventCaptureKeyName)->Empty();
+
+		const FFunctionalTestLoginUser::Result& User1Info = GetOpDataChecked<FFunctionalTestLoginUser::Result>(InAsyncOp, User1KeyName);
+
+		// Create a lobby
+		FCreateLobby::Params Params;
+		Params.LocalUserId = User1Info.AccountInfo->UserId;
+		Params.LocalName = TEXT("test");
+		Params.SchemaName = TEXT("test");
+		Params.MaxMembers = 2;
+		Params.JoinPolicy = ELobbyJoinPolicy::InvitationOnly;
+		//Params.Attributes;
+		Params.LocalUsers.Emplace(FJoinLobbyLocalUserData{User1Info.AccountInfo->UserId, {}});
+		return Params;
+	})
+	.Then(CaptureOperationStepResult<FCreateLobby>(*Op, CreateLobbyKeyName, &FLobbiesCommon::CreateLobby))
+	.Then([this](TOnlineAsyncOp<FFunctionalTestLobbies>& InAsyncOp)
+	{
+		// Check both lobby joined events were received and in the correct order.
+		const TSharedRef<FLobbyEventCapture>& EventCapture = GetOpDataChecked<TSharedRef<FLobbyEventCapture>>(InAsyncOp, LobbyEventCaptureKeyName);
+		if (EventCapture->GetTotalNotificationsReceived() != 2
+			|| EventCapture->LobbyJoined.Num() != 1 || EventCapture->LobbyMemberJoined.Num() != 1
+			|| EventCapture->LobbyJoined[0].GlobalIndex > EventCapture->LobbyMemberJoined[0].GlobalIndex)
+		{
+			InAsyncOp.SetError(Errors::Cancelled());
+		}
+	})
+	.Then([this](TOnlineAsyncOp<FFunctionalTestLobbies>& InAsyncOp)
+	{
+		// Prepare for next test check.
+		GetOpDataChecked<TSharedRef<FLobbyEventCapture>>(InAsyncOp, LobbyEventCaptureKeyName)->Empty();
+
+		const FFunctionalTestLoginUser::Result& User1Info = GetOpDataChecked<FFunctionalTestLoginUser::Result>(InAsyncOp, User1KeyName);
+		const FCreateLobby::Result& CreateLobbyResult = GetOpDataChecked<FCreateLobby::Result>(InAsyncOp, CreateLobbyKeyName);
+
+		FLeaveLobby::Params Params;
+		Params.LocalUserId = User1Info.AccountInfo->UserId;
+		Params.LobbyId = CreateLobbyResult.Lobby->LobbyId;
+		return Params;
+	})
+	.Then(ConsumeOperationStepResult<FLeaveLobby>(*Op, &FLobbiesCommon::LeaveLobby))
+	.Then([this](TOnlineAsyncOp<FFunctionalTestLobbies>& InAsyncOp)
+	{
+		// Check for expected events.
+		const TSharedRef<FLobbyEventCapture>& EventCapture = GetOpDataChecked<TSharedRef<FLobbyEventCapture>>(InAsyncOp, LobbyEventCaptureKeyName);
+		if (EventCapture->GetTotalNotificationsReceived() != 2
+			|| EventCapture->LobbyLeft.Num() != 1 || EventCapture->LobbyMemberLeft.Num() != 1
+			|| EventCapture->LobbyLeft[0].GlobalIndex < EventCapture->LobbyMemberLeft[0].GlobalIndex)
+		{
+			InAsyncOp.SetError(Errors::Cancelled());
+		}
+	})
+
+	//----------------------------------------------------------------------------------------------
+	// Test 2:
+	//    Step 1: Create a lobby with primary user.
+	//    Step 2: Primary user invites secondary user.
+	//    Step 3: Join lobby with secondary user.
+	//    Step 4: Leave lobby with secondary user.
+	//    Step 5: Leave lobby with primary user.
+	//----------------------------------------------------------------------------------------------
+
+	.Then([this](TOnlineAsyncOp<FFunctionalTestLobbies>& InAsyncOp)
+	{
+		// Prepare for next test check.
+		GetOpDataChecked<TSharedRef<FLobbyEventCapture>>(InAsyncOp, LobbyEventCaptureKeyName)->Empty();
+
+		const FFunctionalTestLoginUser::Result& User1Info = GetOpDataChecked<FFunctionalTestLoginUser::Result>(InAsyncOp, User1KeyName);
+
+		FCreateLobby::Params Params;
+		Params.LocalUserId = User1Info.AccountInfo->UserId;
+		Params.LocalName = TEXT("test");
+		Params.SchemaName = TEXT("test");
+		Params.MaxMembers = 2;
+		Params.JoinPolicy = ELobbyJoinPolicy::InvitationOnly;
+		//Params.Attributes;
+		Params.LocalUsers.Emplace(FJoinLobbyLocalUserData{User1Info.AccountInfo->UserId, {}});
+		return Params;
+	})
+	.Then(CaptureOperationStepResult<FCreateLobby>(*Op, CreateLobbyKeyName, &FLobbiesCommon::CreateLobby))
+	.Then([this](TOnlineAsyncOp<FFunctionalTestLobbies>& InAsyncOp)
+	{
+		// Check both lobby joined events were received and in the correct order.
+		const TSharedRef<FLobbyEventCapture>& EventCapture = GetOpDataChecked<TSharedRef<FLobbyEventCapture>>(InAsyncOp, LobbyEventCaptureKeyName);
+		if (EventCapture->GetTotalNotificationsReceived() != 2
+			|| EventCapture->LobbyJoined.Num() != 1 || EventCapture->LobbyMemberJoined.Num() != 1
+			|| EventCapture->LobbyJoined[0].GlobalIndex > EventCapture->LobbyMemberJoined[0].GlobalIndex)
+		{
+			InAsyncOp.SetError(Errors::Cancelled());
+		}
+	})
+	.Then([this](TOnlineAsyncOp<FFunctionalTestLobbies>& InAsyncOp)
+	{
+		// Prepare for next test check.
+		GetOpDataChecked<TSharedRef<FLobbyEventCapture>>(InAsyncOp, LobbyEventCaptureKeyName)->Empty();
+
+		const FFunctionalTestLoginUser::Result& User1Info = GetOpDataChecked<FFunctionalTestLoginUser::Result>(InAsyncOp, User1KeyName);
+		const FFunctionalTestLoginUser::Result& User2Info = GetOpDataChecked<FFunctionalTestLoginUser::Result>(InAsyncOp, User2KeyName);
+		const FCreateLobby::Result& CreateLobbyResult = GetOpDataChecked<FCreateLobby::Result>(InAsyncOp, CreateLobbyKeyName);
+
+		FInviteLobbyMember::Params Params;
+		Params.LocalUserId = User1Info.AccountInfo->UserId;
+		Params.LobbyId = CreateLobbyResult.Lobby->LobbyId;
+		Params.TargetUserId = User2Info.AccountInfo->UserId;
+		return Params;
+	})
+	.Then(ConsumeOperationStepResult<FInviteLobbyMember>(*Op, &FLobbiesCommon::InviteLobbyMember))
+	.Then([this, TestConfig](TOnlineAsyncOp<FFunctionalTestLobbies>& InAsyncOp)
+	{
+		const FFunctionalTestLoginUser::Result& User2Info = GetOpDataChecked<FFunctionalTestLoginUser::Result>(InAsyncOp, User2KeyName);
+		const FCreateLobby::Result& CreateLobbyResult = GetOpDataChecked<FCreateLobby::Result>(InAsyncOp, CreateLobbyKeyName);
+		return AwaitInvitation(User2Info.AccountInfo->UserId, CreateLobbyResult.Lobby->LobbyId, TestConfig->InvitationWaitSeconds);
+	})
+	.Then([this](TOnlineAsyncOp<FFunctionalTestLobbies>& InAsyncOp, TDefaultErrorResultInternal<FOnlineLobbyIdHandle>&& Result)
+	{
+		if (Result.IsError())
+		{
+			InAsyncOp.SetError(Errors::Cancelled(Result.GetErrorValue()));
+		}
+	})
+	.Then([this](TOnlineAsyncOp<FFunctionalTestLobbies>& InAsyncOp)
+	{
+		// Prepare for next test check.
+		GetOpDataChecked<TSharedRef<FLobbyEventCapture>>(InAsyncOp, LobbyEventCaptureKeyName)->Empty();
+
+		const FFunctionalTestLoginUser::Result& User2Info = GetOpDataChecked<FFunctionalTestLoginUser::Result>(InAsyncOp, User2KeyName);
+		const FCreateLobby::Result& CreateLobbyResult = GetOpDataChecked<FCreateLobby::Result>(InAsyncOp, CreateLobbyKeyName);
+
+		FJoinLobby::Params Params;
+		Params.LocalUserId = User2Info.AccountInfo->UserId;
+		Params.LocalName = TEXT("test");
+		Params.LobbyId = CreateLobbyResult.Lobby->LobbyId;
+		Params.LocalUsers.Emplace(FJoinLobbyLocalUserData{User2Info.AccountInfo->UserId, {}});
+		return Params;
+	})
+	.Then(ConsumeOperationStepResult<FJoinLobby>(*Op, &FLobbiesCommon::JoinLobby))
+	.Then([this](TOnlineAsyncOp<FFunctionalTestLobbies>& InAsyncOp)
+	{
+	// todo: currently busted
+	#if 0
+		// Check both lobby joined events were received and in the correct order.
+		const TSharedRef<FLobbyEventCapture>& EventCapture = GetOpDataChecked<TSharedRef<FLobbyEventCapture>>(InAsyncOp, LobbyEventCaptureKeyName);
+		if (EventCapture->GetTotalNotificationsReceived() != 1 || EventCapture->LobbyMemberJoined.Num() != 1)
+		{
+			InAsyncOp.SetError(Errors::Cancelled());
+		}
+	#endif
+	})
+	.Then([this](TOnlineAsyncOp<FFunctionalTestLobbies>& InAsyncOp)
+	{
+		// Prepare for next test check.
+		GetOpDataChecked<TSharedRef<FLobbyEventCapture>>(InAsyncOp, LobbyEventCaptureKeyName)->Empty();
+
+		const FFunctionalTestLoginUser::Result& User2Info = GetOpDataChecked<FFunctionalTestLoginUser::Result>(InAsyncOp, User2KeyName);
+		const FCreateLobby::Result& CreateLobbyResult = GetOpDataChecked<FCreateLobby::Result>(InAsyncOp, CreateLobbyKeyName);
+
+		FLeaveLobby::Params Params;
+		Params.LocalUserId = User2Info.AccountInfo->UserId;
+		Params.LobbyId = CreateLobbyResult.Lobby->LobbyId;
+		return Params;
+	})
+	.Then(ConsumeOperationStepResult<FLeaveLobby>(*Op, &FLobbiesCommon::LeaveLobby))
+	.Then([this](TOnlineAsyncOp<FFunctionalTestLobbies>& InAsyncOp)
+	{
+	// todo: currently busted
+	#if 0
+		// Check for expected events.
+		const TSharedRef<FLobbyEventCapture>& EventCapture = GetOpDataChecked<TSharedRef<FLobbyEventCapture>>(InAsyncOp, LobbyEventCaptureKeyName);
+		if (EventCapture->GetTotalNotificationsReceived() != 1|| EventCapture->LobbyMemberLeft.Num())
+		{
+			InAsyncOp.SetError(Errors::Cancelled());
+		}
+	#endif
+	})
+	.Then([this](TOnlineAsyncOp<FFunctionalTestLobbies>& InAsyncOp)
+	{
+		// Prepare for next test check.
+		GetOpDataChecked<TSharedRef<FLobbyEventCapture>>(InAsyncOp, LobbyEventCaptureKeyName)->Empty();
+
+		const FFunctionalTestLoginUser::Result& User1Info = GetOpDataChecked<FFunctionalTestLoginUser::Result>(InAsyncOp, User1KeyName);
+		const FCreateLobby::Result& CreateLobbyResult = GetOpDataChecked<FCreateLobby::Result>(InAsyncOp, CreateLobbyKeyName);
+
+		FLeaveLobby::Params Params;
+		Params.LocalUserId = User1Info.AccountInfo->UserId;
+		Params.LobbyId = CreateLobbyResult.Lobby->LobbyId;
+		return Params;
+	})
+	.Then(ConsumeOperationStepResult<FLeaveLobby>(*Op, &FLobbiesCommon::LeaveLobby))
+	.Then([this](TOnlineAsyncOp<FFunctionalTestLobbies>& InAsyncOp)
+	{
+	// todo: currently busted
+	#if 0
+		// Check for expected events.
+		const TSharedRef<FLobbyEventCapture>& EventCapture = GetOpDataChecked<TSharedRef<FLobbyEventCapture>>(InAsyncOp, LobbyEventCaptureKeyName);
+		if (EventCapture->GetTotalNotificationsReceived() != 2
+			|| EventCapture->LobbyLeft.Num() != 1 || EventCapture->LobbyMemberLeft.Num() != 1
+			|| EventCapture->LobbyLeft[0].GlobalIndex < EventCapture->LobbyMemberLeft[0].GlobalIndex)
+		{
+			InAsyncOp.SetError(Errors::Cancelled());
+		}
+	#endif
+	})
+
+	//----------------------------------------------------------------------------------------------
+	// Test 3: Simple search
+	//    Step 1: Create a lobby with primary user.
+	//    Step 2: Secondary user searches for lobbies.
+	//    Step 3: Join lobby with secondary user.
+	//    Step 4: Leave lobby with secondary user.
+	//    Step 5: Leave lobby with primary user.
+	//----------------------------------------------------------------------------------------------
+
+	.Then([this](TOnlineAsyncOp<FFunctionalTestLobbies>& InAsyncOp)
+	{
+		// Prepare for next test check.
+		GetOpDataChecked<TSharedRef<FLobbyEventCapture>>(InAsyncOp, LobbyEventCaptureKeyName)->Empty();
+
+		const FFunctionalTestLoginUser::Result& User1Info = GetOpDataChecked<FFunctionalTestLoginUser::Result>(InAsyncOp, User1KeyName);
+
+		FCreateLobby::Params Params;
+		Params.LocalUserId = User1Info.AccountInfo->UserId;
+		Params.LocalName = TEXT("test");
+		Params.SchemaName = TEXT("test");
+		Params.MaxMembers = 2;
+		Params.JoinPolicy = ELobbyJoinPolicy::PublicAdvertised;
+		//Params.Attributes;
+		Params.LocalUsers.Emplace(FJoinLobbyLocalUserData{User1Info.AccountInfo->UserId, {}});
+		return Params;
+	})
+	.Then(CaptureOperationStepResult<FCreateLobby>(*Op, CreateLobbyKeyName, &FLobbiesCommon::CreateLobby))
+	.Then([this](TOnlineAsyncOp<FFunctionalTestLobbies>& InAsyncOp)
+	{
+		// Check both lobby joined events were received and in the correct order.
+		const TSharedRef<FLobbyEventCapture>& EventCapture = GetOpDataChecked<TSharedRef<FLobbyEventCapture>>(InAsyncOp, LobbyEventCaptureKeyName);
+		if (EventCapture->GetTotalNotificationsReceived() != 2
+			|| EventCapture->LobbyJoined.Num() != 1 || EventCapture->LobbyMemberJoined.Num() != 1
+			|| EventCapture->LobbyJoined[0].GlobalIndex > EventCapture->LobbyMemberJoined[0].GlobalIndex)
+		{
+			InAsyncOp.SetError(Errors::Cancelled());
+		}
+	})
+	.Then([this](TOnlineAsyncOp<FFunctionalTestLobbies>& InAsyncOp)
+	{
+		const FFunctionalTestLoginUser::Result& User2Info = GetOpDataChecked<FFunctionalTestLoginUser::Result>(InAsyncOp, User2KeyName);
+		const FCreateLobby::Result& CreateLobbyResult = GetOpDataChecked<FCreateLobby::Result>(InAsyncOp, CreateLobbyKeyName);
+
+		// Search for lobby from create
+		FFindLobby::Params Params;
+		Params.LocalUserId = User2Info.AccountInfo->UserId;
+		Params.LobbyId = CreateLobbyResult.Lobby->LobbyId;
+		return Params;
+	})
+	.Then(CaptureOperationStepResult<FFindLobby>(*Op, FindLobbyKeyName, &FLobbiesCommon::FindLobby))
+	.Then([this](TOnlineAsyncOp<FFunctionalTestLobbies>& InAsyncOp)
+	{
+		// Prepare for next test check.
+		GetOpDataChecked<TSharedRef<FLobbyEventCapture>>(InAsyncOp, LobbyEventCaptureKeyName)->Empty();
+
+		const FFunctionalTestLoginUser::Result& User2Info = GetOpDataChecked<FFunctionalTestLoginUser::Result>(InAsyncOp, User2KeyName);
+		const FCreateLobby::Result& CreateLobbyResult = GetOpDataChecked<FCreateLobby::Result>(InAsyncOp, CreateLobbyKeyName);
+		const FFindLobby::Result& FindResults = GetOpDataChecked<FFindLobby::Result>(InAsyncOp, FindLobbyKeyName);
+
+		// Not a valid expectation without a custom bucket.
+		if (FindResults.Lobbies.Num() != 1)
+		{
+			InAsyncOp.SetError(Errors::Cancelled());
+			return FJoinLobby::Params{};
+		}
+
+		FJoinLobby::Params Params;
+		Params.LocalUserId = User2Info.AccountInfo->UserId;
+		Params.LocalName = TEXT("test");
+		Params.LobbyId = FindResults.Lobbies[0]->LobbyId;
+		Params.LocalUsers.Emplace(FJoinLobbyLocalUserData{User2Info.AccountInfo->UserId, {}});
+		return Params;
+	})
+	.Then(ConsumeOperationStepResult<FJoinLobby>(*Op, &FLobbiesCommon::JoinLobby))
+	.Then([this](TOnlineAsyncOp<FFunctionalTestLobbies>& InAsyncOp)
+	{
+	// todo: currently busted
+	#if 0
+		// Check both lobby joined events were received and in the correct order.
+		const TSharedRef<FLobbyEventCapture>& EventCapture = GetOpDataChecked<TSharedRef<FLobbyEventCapture>>(InAsyncOp, LobbyEventCaptureKeyName);
+		if (EventCapture->GetTotalNotificationsReceived() != 1 || EventCapture->LobbyMemberJoined.Num() != 1)
+		{
+			InAsyncOp.SetError(Errors::Cancelled());
+		}
+	#endif
+	})
+	.Then([this](TOnlineAsyncOp<FFunctionalTestLobbies>& InAsyncOp)
+	{
+		// Prepare for next test check.
+		GetOpDataChecked<TSharedRef<FLobbyEventCapture>>(InAsyncOp, LobbyEventCaptureKeyName)->Empty();
+
+		const FFunctionalTestLoginUser::Result& User2Info = GetOpDataChecked<FFunctionalTestLoginUser::Result>(InAsyncOp, User2KeyName);
+		const FCreateLobby::Result& CreateLobbyResult = GetOpDataChecked<FCreateLobby::Result>(InAsyncOp, CreateLobbyKeyName);
+
+		FLeaveLobby::Params Params;
+		Params.LocalUserId = User2Info.AccountInfo->UserId;
+		Params.LobbyId = CreateLobbyResult.Lobby->LobbyId;
+		return Params;
+	})
+	.Then(ConsumeOperationStepResult<FLeaveLobby>(*Op, &FLobbiesCommon::LeaveLobby))
+	.Then([this](TOnlineAsyncOp<FFunctionalTestLobbies>& InAsyncOp)
+	{
+	// todo: currently busted
+	#if 0
+		// Check for expected events.
+		const TSharedRef<FLobbyEventCapture>& EventCapture = GetOpDataChecked<TSharedRef<FLobbyEventCapture>>(InAsyncOp, LobbyEventCaptureKeyName);
+		if (EventCapture->GetTotalNotificationsReceived() != 1|| EventCapture->LobbyMemberLeft.Num())
+		{
+			InAsyncOp.SetError(Errors::Cancelled());
+		}
+	#endif
+	})
+	.Then([this](TOnlineAsyncOp<FFunctionalTestLobbies>& InAsyncOp)
+	{
+		// Prepare for next test check.
+		GetOpDataChecked<TSharedRef<FLobbyEventCapture>>(InAsyncOp, LobbyEventCaptureKeyName)->Empty();
+
+		const FFunctionalTestLoginUser::Result& User1Info = GetOpDataChecked<FFunctionalTestLoginUser::Result>(InAsyncOp, User1KeyName);
+		const FCreateLobby::Result& CreateLobbyResult = GetOpDataChecked<FCreateLobby::Result>(InAsyncOp, CreateLobbyKeyName);
+
+		FLeaveLobby::Params Params;
+		Params.LocalUserId = User1Info.AccountInfo->UserId;
+		Params.LobbyId = CreateLobbyResult.Lobby->LobbyId;
+		return Params;
+	})
+	.Then(ConsumeOperationStepResult<FLeaveLobby>(*Op, &FLobbiesCommon::LeaveLobby))
+	.Then([this](TOnlineAsyncOp<FFunctionalTestLobbies>& InAsyncOp)
+	{
+	// todo: currently busted
+	#if 0
+		// Check for expected events.
+		const TSharedRef<FLobbyEventCapture>& EventCapture = GetOpDataChecked<TSharedRef<FLobbyEventCapture>>(InAsyncOp, LobbyEventCaptureKeyName);
+		if (EventCapture->GetTotalNotificationsReceived() != 2
+			|| EventCapture->LobbyLeft.Num() != 1 || EventCapture->LobbyMemberLeft.Num() != 1
+			|| EventCapture->LobbyLeft[0].GlobalIndex < EventCapture->LobbyMemberLeft[0].GlobalIndex)
+		{
+			InAsyncOp.SetError(Errors::Cancelled());
+		}
+	#endif
+	})
+
+	//----------------------------------------------------------------------------------------------
+	// Complete
+	//----------------------------------------------------------------------------------------------
+
+	.Then([this](TOnlineAsyncOp<FFunctionalTestLobbies>& InAsyncOp)
+	{
+		InAsyncOp.SetResult(FFunctionalTestLobbies::Result{});
+	})
+	.Enqueue(GetServices().GetParallelQueue());
+
+	return Op->GetHandle();
+}
+
+TFuture<TOnlineResult<FLobbiesCommon::FFunctionalTestLoginUser>> FLobbiesCommon::FunctionalTestLoginUser(FFunctionalTestLoginUser::Params&& Params)
+{
+	UE_LOG(LogLobbies, Log, TEXT("[FLobbiesCommon::FunctionalTestLoginUser] Logging in user. PlatformUserId: %d"), FPlatformMisc::GetUserIndexForPlatformUser(Params.PlatformUserId));
+
+	IAuth* AuthInterface = Services.Get<IAuth>();
+	if (!AuthInterface)
+	{
+		UE_LOG(LogLobbies, Warning, TEXT("[FLobbiesCommon::FunctionalTestLoginUser] Login failed. PlatformUserId: %d"), FPlatformMisc::GetUserIndexForPlatformUser(Params.PlatformUserId));
+		return MakeFulfilledPromise<TOnlineResult<FFunctionalTestLoginUser>>(Errors::MissingInterface()).GetFuture();
+	}
+
+	TSharedRef<TPromise<TOnlineResult<FFunctionalTestLoginUser>>> Promise = MakeShared<TPromise<TOnlineResult<FFunctionalTestLoginUser>>>();
+	TFuture<TOnlineResult<FFunctionalTestLoginUser>> Future = Promise->GetFuture();
+
+	FAuthLogin::Params LoginParams;
+	LoginParams.PlatformUserId = Params.PlatformUserId;
+	LoginParams.CredentialsType = Params.Type;
+	LoginParams.CredentialsId = Params.Id;
+	LoginParams.CredentialsToken.Set<FString>(Params.Token);
+
+	AuthInterface->Login(MoveTemp(LoginParams))
+	.OnComplete([Promise, PlatformUserId = Params.PlatformUserId](const TOnlineResult<FAuthLogin>& LoginResult)
+	{
+		if (LoginResult.IsError())
+		{
+			UE_LOG(LogLobbies, Warning, TEXT("[FLobbiesCommon::FunctionalTestLoginUser] Login failed. PlatformUserId: %d, Result: %s"), FPlatformMisc::GetUserIndexForPlatformUser(PlatformUserId), *LoginResult.GetErrorValue().GetLogString());
+			Promise->EmplaceValue(Errors::RequestFailure());
+		}
+		else
+		{
+			UE_LOG(LogLobbies, Log, TEXT("[FLobbiesCommon::FunctionalTestLoginUser] Login Succeeded. PlatformUserId: %d, User: %s"), FPlatformMisc::GetUserIndexForPlatformUser(PlatformUserId), *ToLogString(LoginResult.GetOkValue().AccountInfo->UserId));
+			Promise->EmplaceValue(FFunctionalTestLoginUser::Result{LoginResult.GetOkValue().AccountInfo});
+		}
+	});
+
+	return Future;
+}
+
+TFuture<TOnlineResult<FLobbiesCommon::FFunctionalTestLogoutUser>> FLobbiesCommon::FunctionalTestLogoutUser(FFunctionalTestLogoutUser::Params&& Params)
+{
+	UE_LOG(LogLobbies, Log, TEXT("[FLobbiesCommon::FunctionalTestLogoutUser] Logging out user. PlatformUserId: %d"), FPlatformMisc::GetUserIndexForPlatformUser(Params.PlatformUserId));
+
+	IAuth* AuthInterface = Services.Get<IAuth>();
+	if (!AuthInterface)
+	{
+		UE_LOG(LogLobbies, Warning, TEXT("[FLobbiesCommon::FunctionalTestLogoutUser] Logout failed. PlatformUserId: %d"), FPlatformMisc::GetUserIndexForPlatformUser(Params.PlatformUserId));
+		return MakeFulfilledPromise<TOnlineResult<FFunctionalTestLogoutUser>>(Errors::MissingInterface()).GetFuture();
+	}
+
+	FAuthGetAccountByPlatformUserId::Params GetAccountParams;
+	GetAccountParams.PlatformUserId = Params.PlatformUserId;
+	TOnlineResult<FAuthGetAccountByPlatformUserId> Result = AuthInterface->GetAccountByPlatformUserId(MoveTemp(GetAccountParams));
+	if (Result.IsError())
+	{
+		// Ignore errors for now. This function should ignore logged out users.
+		if (Result.GetErrorValue() == Errors::Unknown())
+		{
+			UE_LOG(LogLobbies, Log, TEXT("[FLobbiesCommon::FunctionalTestLogoutUser] Logout Succeeded - already logged out. PlatformUserId: %d"), FPlatformMisc::GetUserIndexForPlatformUser(Params.PlatformUserId));
+			return MakeFulfilledPromise<TOnlineResult<FFunctionalTestLogoutUser>>(FFunctionalTestLogoutUser::Result{}).GetFuture();
+		}
+		else
+		{
+			UE_LOG(LogLobbies, Warning, TEXT("[FLobbiesCommon::FunctionalTestLogoutUser] Logout failed. PlatformUserId: %d, Result: %s"), FPlatformMisc::GetUserIndexForPlatformUser(Params.PlatformUserId), *Result.GetErrorValue().GetLogString());
+			return MakeFulfilledPromise<TOnlineResult<FFunctionalTestLogoutUser>>(Errors::Unknown(MoveTemp(Result.GetErrorValue()))).GetFuture();
+		}
+	}
+
+	TSharedRef<FAccountInfo> AccountInfo = Result.GetOkValue().AccountInfo;
+	if (AccountInfo->LoginStatus == ELoginStatus::NotLoggedIn)
+	{
+		return MakeFulfilledPromise<TOnlineResult<FFunctionalTestLogoutUser>>(Errors::InvalidUser()).GetFuture();
+	}
+
+	TSharedRef<TPromise<TOnlineResult<FFunctionalTestLogoutUser>>> Promise = MakeShared<TPromise<TOnlineResult<FFunctionalTestLogoutUser>>>();
+	TFuture<TOnlineResult<FFunctionalTestLogoutUser>> Future = Promise->GetFuture();
+
+	FAuthLogout::Params LogoutParams;
+	LogoutParams.LocalUserId = AccountInfo->UserId;
+
+	AuthInterface->Logout(MoveTemp(LogoutParams))
+	.OnComplete([Promise, PlatformUserId = Params.PlatformUserId](const TOnlineResult<FAuthLogout>& LogoutResult) mutable
+	{
+		if (LogoutResult.IsError())
+		{
+			UE_LOG(LogLobbies, Warning, TEXT("[FLobbiesCommon::FunctionalTestLogoutUser] Logout failed. PlatformUserId: %d, Result: %s"), FPlatformMisc::GetUserIndexForPlatformUser(PlatformUserId), *LogoutResult.GetErrorValue().GetLogString());
+			Promise->EmplaceValue(Errors::RequestFailure());
+		}
+		else
+		{
+			UE_LOG(LogLobbies, Log, TEXT("[FLobbiesCommon::FunctionalTestLogoutUser] Logout Succeeded. PlatformUserId: %d"), FPlatformMisc::GetUserIndexForPlatformUser(PlatformUserId));
+			Promise->EmplaceValue(FFunctionalTestLogoutUser::Result{});
+		}
+	});
+
+	return Future;
+}
+
+TFuture<TOnlineResult<FLobbiesCommon::FFunctionalTestLogoutAllUsers>> FLobbiesCommon::FunctionalTestLogoutAllUsers(FFunctionalTestLogoutAllUsers::Params&& Params)
+{
+	UE_LOG(LogLobbies, Log, TEXT("[FLobbiesCommon::FFunctionalTestLogoutAllUsers] Logging out all users."));
+
+	TArray<TFuture<TOnlineResult<FLobbiesCommon::FFunctionalTestLogoutUser>>> LogoutFutures;
+	for (int32 index = 0; index < MAX_LOCAL_PLAYERS; ++index)
+	{
+		LogoutFutures.Emplace(FunctionalTestLogoutUser(FFunctionalTestLogoutUser::Params{FPlatformMisc::GetPlatformUserForUserIndex(index)}));
+	}
+
+	TSharedRef<TPromise<TOnlineResult<FFunctionalTestLogoutAllUsers>>> Promise = MakeShared<TPromise<TOnlineResult<FFunctionalTestLogoutAllUsers>>>();
+	TFuture<TOnlineResult<FFunctionalTestLogoutAllUsers>> Future = Promise->GetFuture();
+
+	WhenAll(MoveTempIfPossible(LogoutFutures))
+	.Then([Promise = MoveTemp(Promise)](TFuture<TArray<TOnlineResult<FLobbiesCommon::FFunctionalTestLogoutUser>>>&& Results)
+	{
+		bool HasAnyError = false;
+		for (TOnlineResult<FLobbiesCommon::FFunctionalTestLogoutUser>& Result : Results.Get())
+		{
+			HasAnyError |= Result.IsError();
+		}
+
+		if (HasAnyError)
+		{
+			UE_LOG(LogLobbies, Warning, TEXT("[FLobbiesCommon::FunctionalTestLogoutUser] Logout all users Failed."));
+			Promise->EmplaceValue(Errors::RequestFailure());
+		}
+		else
+		{
+			UE_LOG(LogLobbies, Log, TEXT("[FLobbiesCommon::FFunctionalTestLogoutAllUsers] Logout all users Succeeded."));
+			Promise->EmplaceValue(FFunctionalTestLogoutAllUsers::Result{});
+		}
+	});
+
+	return Future;
+}
+
+#endif
+
+/* UE::Online */ }
diff --git a/Engine/Plugins/Online/Experimental/OnlineServices/Source/OnlineServicesCommon/Public/Online/LobbiesCommon.h b/Engine/Plugins/Online/Experimental/OnlineServices/Source/OnlineServicesCommon/Public/Online/LobbiesCommon.h
new file mode 100644
--- /dev/null
+++ b/Engine/Plugins/Online/Experimental/OnlineServices/Source/OnlineServicesCommon/Public/Online/LobbiesCommon.h
@@ -0,0 +1,555 @@
+// Copyright Epic Games, Inc. All Rights Reserved.
+
+#pragma once
+
+#include "Online/Lobbies.h"
+#include "Online/OnlineComponent.h"
+
+#define LOBBIES_FUNCTIONAL_TEST_ENABLED !UE_BUILD_SHIPPING
+
+DECLARE_LOG_CATEGORY_EXTERN(LogLobbies, Log, All);
+
+namespace UE::Online {
+
+class FOnlineServicesCommon;
+class FAccountInfo;
+
+#if LOBBIES_FUNCTIONAL_TEST_ENABLED
+struct FFunctionalTestLobbies
+{
+	static constexpr TCHAR Name[] = TEXT("FunctionalTest");
+
+	struct Params
+	{
+	};
+
+	struct Result
+	{
+	};
+};
+#endif
+
+struct FLobbyConfig
+{
+	TArray<FLobbySchemaId> RegisteredSchemas;
+};
+
+struct FLobbySchemaAttributeConfig
+{
+	// The attribute name.
+	FLobbyAttributeId Name;
+	
+	// Override default service attribute field assignment.
+	// Specifying an override allows custom grouping of attributes within a
+	// platform attribute. This grouping is useful when a set of attributes are
+	// known to update at the same cadence.
+	TOptional<FLobbyAttributeId> ServiceAttributeFieldId;
+	
+	// The size in bytes of the attribute value.
+	// A schema will fail validation if all attributes will not fit within the
+	// platform's lobby fields.
+	uint32 MaxByteSize = 0;
+
+	// Control whether the attribute is visible to players who are not joined
+	// to the lobby.
+	ELobbyAttributeVisibility Visibility = ELobbyAttributeVisibility::Private;
+
+};
+
+struct FLobbySchemaConfig
+{
+	// The schema name.
+	FLobbySchemaId SchemaName;
+
+	// The optional base schema name.
+	FLobbySchemaId BaseSchemaName;
+	
+	// The definitions for attributes attached to a lobby.
+	TArray<FLobbySchemaAttributeConfig> LobbyAttributes;
+
+	// The definitions for attributes attached to a lobby member.
+	TArray<FLobbySchemaAttributeConfig> LobbyMemberAttributes;
+};
+
+class FLobbyServiceAttributeChanges final
+{
+public:
+	TMap<FLobbyAttributeId, FLobbyVariant> MutatedAttributes;
+	TSet<FLobbyAttributeId> ClearedAttributes;
+};
+
+class FLobbyClientAttributeChanges final
+{
+public:
+	// Apply attribute changes to attribute parameter.
+	// Applying the changes will clear MutatedAttributes and ClearedAttributes.
+	// Returns a list of the attributes which were changed.
+	TSet<FLobbyAttributeId> Apply(TMap<FLobbyAttributeId, FLobbyVariant>& InOutAttributes)
+	{
+		TSet<FLobbyAttributeId> Changes;
+		Changes.Reserve(MutatedAttributes.Num() + ClearedAttributes.Num());
+
+		for (TPair<FLobbyAttributeId, FLobbyVariant>& Attribute : MutatedAttributes)
+		{
+			Changes.Add(Attribute.Key);
+			InOutAttributes.Emplace(Attribute.Key, MoveTemp(Attribute.Value));
+		}
+
+		for (FLobbyAttributeId AttributeId : ClearedAttributes)
+		{
+			Changes.Add(AttributeId);
+			InOutAttributes.Remove(AttributeId);
+		}
+
+		MutatedAttributes.Empty();
+		ClearedAttributes.Empty();
+
+		return Changes;
+	}
+
+	TMap<FLobbyAttributeId, FLobbyVariant> MutatedAttributes;
+	TSet<FLobbyAttributeId> ClearedAttributes;
+};
+
+class FLobbySchema final
+{
+public:
+	static TSharedPtr<FLobbySchema> Create(FLobbySchemaConfig LobbySchemaConfig);
+
+	TSharedRef<FLobbyServiceAttributeChanges> TranslateLobbyAttributes(const TMap<FLobbyAttributeId, FLobbyVariant>& LobbyAttributes, const FLobbyClientAttributeChanges& ClientAttributeChanges) const;
+	TSharedRef<FLobbyClientAttributeChanges> TranslateLobbyAttributes(const FLobbyServiceAttributeChanges& ServiceAttributeChanges) const;
+	TSharedRef<FLobbyServiceAttributeChanges> TranslateLobbyMemberAttributes(const TMap<FLobbyAttributeId, FLobbyVariant>& LobbyAttributes, const FLobbyClientAttributeChanges& ClientAttributeChanges) const;
+	TSharedRef<FLobbyClientAttributeChanges> TranslateLobbyMemberAttributes(const FLobbyServiceAttributeChanges& ServiceAttributeChanges) const;
+
+private:
+	FLobbySchemaId SchemaId;
+};
+
+class FLobbySchemaRegistry
+{
+public:
+	bool Initialize(TArray<FLobbySchemaConfig> LobbySchemaConfigs);
+	TSharedPtr<FLobbySchema> FindSchema(FLobbySchemaId SchemaId);
+
+private:
+	bool RegisterSchema(FLobbySchemaConfig LobbySchemaConfig);
+
+	TMap<FLobbySchemaId, TSharedRef<FLobbySchema>> RegisteredSchemas;
+};
+
+// Todo: put this somewhere else
+template <typename AwaitedType>
+TFuture<TArray<AwaitedType>> WhenAll(TArray<TFuture<AwaitedType>>&& Futures)
+{
+	struct FWhenAllState
+	{
+		TArray<TFuture<AwaitedType>> Futures;
+		TArray<AwaitedType> Results;
+		TPromise<TArray<AwaitedType>> FinalPromise;
+	};
+
+	if (Futures.IsEmpty())
+	{
+		return MakeFulfilledPromise<TArray<AwaitedType>>().GetFuture();
+	}
+	else
+	{
+		TSharedRef<FWhenAllState> WhenAllState = MakeShared<FWhenAllState>();
+		WhenAllState->Futures = MoveTemp(Futures);
+
+		for (TFuture<AwaitedType>& Future : WhenAllState->Futures)
+		{
+			Future.Then([WhenAllState](TFuture<AwaitedType>&& AwaitedResult)
+			{
+				WhenAllState->Results.Emplace(MoveTempIfPossible(AwaitedResult.Get()));
+
+				if (WhenAllState->Futures.Num() == WhenAllState->Results.Num())
+				{
+					WhenAllState->FinalPromise.EmplaceValue(MoveTemp(WhenAllState->Results));
+				}
+			});
+		}
+
+		return WhenAllState->FinalPromise.GetFuture();
+	}
+}
+
+class ONLINESERVICESCOMMON_API FLobbiesCommon : public TOnlineComponent<ILobbies>
+{
+public:
+	using Super = ILobbies;
+
+	FLobbiesCommon(FOnlineServicesCommon& InServices);
+
+	// TOnlineComponent
+	virtual void Initialize() override;
+	virtual void RegisterCommands() override;
+
+	// ILobbies
+	virtual TOnlineAsyncOpHandle<FCreateLobby> CreateLobby(FCreateLobby::Params&& Params) override;
+	virtual TOnlineAsyncOpHandle<FFindLobby> FindLobby(FFindLobby::Params&& Params) override;
+	virtual TOnlineAsyncOpHandle<FRestoreLobbies> RestoreLobbies(FRestoreLobbies::Params&& Params) override;
+	virtual TOnlineAsyncOpHandle<FJoinLobby> JoinLobby(FJoinLobby::Params&& Params) override;
+	virtual TOnlineAsyncOpHandle<FLeaveLobby> LeaveLobby(FLeaveLobby::Params&& Params) override;
+	virtual TOnlineAsyncOpHandle<FInviteLobbyMember> InviteLobbyMember(FInviteLobbyMember::Params&& Params) override;
+	virtual TOnlineAsyncOpHandle<FDeclineLobbyInvitation> DeclineLobbyInvitation(FDeclineLobbyInvitation::Params&& Params) override;
+	virtual TOnlineAsyncOpHandle<FKickLobbyMember> KickLobbyMember(FKickLobbyMember::Params&& Params) override;
+	virtual TOnlineAsyncOpHandle<FPromoteLobbyMember> PromoteLobbyMember(FPromoteLobbyMember::Params&& Params) override;
+	virtual TOnlineAsyncOpHandle<FModifyLobbySchema> ModifyLobbySchema(FModifyLobbySchema::Params&& Params) override;
+	virtual TOnlineAsyncOpHandle<FModifyLobbyJoinPolicy> ModifyLobbyJoinPolicy(FModifyLobbyJoinPolicy::Params&& Params) override;
+	virtual TOnlineAsyncOpHandle<FModifyLobbyAttributes> ModifyLobbyAttributes(FModifyLobbyAttributes::Params&& Params) override;
+	virtual TOnlineAsyncOpHandle<FModifyLobbyMemberAttributes> ModifyLobbyMemberAttributes(FModifyLobbyMemberAttributes::Params&& Params) override;
+	virtual TOnlineResult<FGetJoinedLobbies> GetJoinedLobbies(FGetJoinedLobbies::Params&& Params) override;
+	virtual TOnlineResult<FGetReceivedInvitations> GetReceivedInvitations(FGetReceivedInvitations::Params&& Params) override;
+
+	virtual TOnlineEvent<void(const FLobbyJoined&)> OnLobbyJoined() override;
+	virtual TOnlineEvent<void(const FLobbyLeft&)> OnLobbyLeft() override;
+	virtual TOnlineEvent<void(const FLobbyMemberJoined&)> OnLobbyMemberJoined() override;
+	virtual TOnlineEvent<void(const FLobbyMemberLeft&)> OnLobbyMemberLeft() override;
+	virtual TOnlineEvent<void(const FLobbyLeaderChanged&)> OnLobbyLeaderChanged() override;
+	virtual TOnlineEvent<void(const FLobbySchemaChanged&)> OnLobbySchemaChanged() override;
+	virtual TOnlineEvent<void(const FLobbyAttributesChanged&)> OnLobbyAttributesChanged() override;
+	virtual TOnlineEvent<void(const FLobbyMemberAttributesChanged&)> OnLobbyMemberAttributesChanged() override;
+	virtual TOnlineEvent<void(const FLobbyInvitationAdded&)> OnLobbyInvitationAdded() override;
+	virtual TOnlineEvent<void(const FLobbyInvitationRemoved&)> OnLobbyInvitationRemoved() override;
+
+protected:
+	TFuture<TDefaultErrorResultInternal<FOnlineLobbyIdHandle>> AwaitInvitation(
+		FOnlineAccountIdHandle TargetAccountId,
+		FOnlineLobbyIdHandle LobbyId,
+		float TimeoutSeconds);
+
+#if LOBBIES_FUNCTIONAL_TEST_ENABLED
+	TOnlineAsyncOpHandle<FFunctionalTestLobbies> FunctionalTest(FFunctionalTestLobbies::Params&& Params);
+
+	struct FFunctionalTestLoginUser
+	{
+		static constexpr TCHAR Name[] = TEXT("FunctionalTestLoginUser");
+
+		struct Params
+		{
+			FPlatformUserId PlatformUserId;
+			FString Type;
+			FString Id;
+			FString Token;
+		};
+
+		struct Result
+		{
+			TSharedPtr<FAccountInfo> AccountInfo;
+		};
+	};
+	TFuture<TOnlineResult<FFunctionalTestLoginUser>> FunctionalTestLoginUser(FFunctionalTestLoginUser::Params&& Params);
+
+	struct FFunctionalTestLogoutUser
+	{
+		static constexpr TCHAR Name[] = TEXT("FunctionalTestLogoutUser");
+
+		struct Params
+		{
+			FPlatformUserId PlatformUserId;
+		};
+
+		struct Result
+		{
+		};
+	};
+	TFuture<TOnlineResult<FFunctionalTestLogoutUser>> FunctionalTestLogoutUser(FFunctionalTestLogoutUser::Params&& Params);
+
+	struct FFunctionalTestLogoutAllUsers
+	{
+		static constexpr TCHAR Name[] = TEXT("FunctionalTestLogoutAllUsers");
+
+		struct Params
+		{
+		};
+
+		struct Result
+		{
+		};
+	};
+	TFuture<TOnlineResult<FFunctionalTestLogoutAllUsers>> FunctionalTestLogoutAllUsers(FFunctionalTestLogoutAllUsers::Params&& Params);
+#endif
+
+	template <typename SecondaryOpType, typename OpType, typename Function>
+	TFunction<TFuture<void>(TOnlineAsyncOp<OpType>&, typename SecondaryOpType::Params&&)> ConsumeStepResult(TOnlineAsyncOp<OpType>& InAsyncOp, Function Func)
+	{
+		return [this, Func](TOnlineAsyncOp<OpType>& InAsyncOp, typename SecondaryOpType::Params&& InParams) mutable
+		{
+			TPromise<void> Promise;
+			auto Future = Promise.GetFuture();
+
+			UE_LOG(LogLobbies, Log, TEXT("[FLobbiesCommon::ConsumeStepResult] Starting secondary operation %s."), SecondaryOpType::Name);
+
+			(this->*Func)(MoveTempIfPossible(InParams))
+			.Then([Promise = MoveTemp(Promise), Op = InAsyncOp.AsShared()](TFuture<TOnlineResult<SecondaryOpType>>&& Future) mutable
+			{
+				if (Future.Get().IsError())
+				{
+					UE_LOG(LogLobbies, Log, TEXT("[FLobbiesCommon::ConsumeStepResult] Failed secondary operation %s. Error: %s"), SecondaryOpType::Name, *Future.Get().GetErrorValue().GetLogString());
+					Op->SetError(Errors::RequestFailure(MoveTempIfPossible(Future.Get().GetErrorValue())));
+				}
+				else
+				{
+					UE_LOG(LogLobbies, Log, TEXT("[FLobbiesCommon::ConsumeStepResult] Completed secondary operation %s."), SecondaryOpType::Name);
+				}
+				Promise.EmplaceValue();
+			});
+
+			return Future;
+		};
+	}
+
+	template <typename SecondaryOpType, typename OpType, typename Function>
+	TFunction<TFuture<void>(TOnlineAsyncOp<OpType>&)> ConsumeStepResult(TOnlineAsyncOp<OpType>& InAsyncOp, Function Func, typename SecondaryOpType::Params&& InParams)
+	{
+		return [this, Func, InParams = MoveTempIfPossible(InParams)](TOnlineAsyncOp<OpType>& InAsyncOp) mutable
+		{
+			TPromise<void> Promise;
+			auto Future = Promise.GetFuture();
+
+			UE_LOG(LogLobbies, Log, TEXT("[FLobbiesCommon::ConsumeStepResult] Starting secondary operation %s."), SecondaryOpType::Name);
+
+			(this->*Func)(MoveTempIfPossible(InParams))
+			.Then([Promise = MoveTemp(Promise), Op = InAsyncOp.AsShared()](TFuture<TOnlineResult<SecondaryOpType>>&& Future) mutable
+			{
+				if (Future.Get().IsError())
+				{
+					UE_LOG(LogLobbies, Log, TEXT("[FLobbiesCommon::ConsumeStepResult] Failed secondary operation %s. Error: %s"), SecondaryOpType::Name, *Future.Get().GetErrorValue().GetLogString());
+					Op->SetError(Errors::RequestFailure(MoveTempIfPossible(Future.Get().GetErrorValue())));
+				}
+				else
+				{
+					UE_LOG(LogLobbies, Log, TEXT("[FLobbiesCommon::ConsumeStepResult] Completed secondary operation %s."), SecondaryOpType::Name);
+				}
+				Promise.EmplaceValue();
+			});
+
+			return Future;
+		};
+	}
+
+	template <typename SecondaryOpType, typename OpType, typename Function>
+	TFunction<TFuture<void>(TOnlineAsyncOp<OpType>&, typename SecondaryOpType::Params&&)> CaptureStepResult(TOnlineAsyncOp<OpType>& InAsyncOp, const FString& ResultKey, Function Func)
+	{
+		return [this, ResultKey, Func](TOnlineAsyncOp<OpType>& InAsyncOp, typename SecondaryOpType::Params&& InParams) mutable
+		{
+			TPromise<void> Promise;
+			auto Future = Promise.GetFuture();
+
+			UE_LOG(LogLobbies, Log, TEXT("[FLobbiesCommon::CaptureStepResult] Starting secondary operation %s."), SecondaryOpType::Name);
+
+			(this->*Func)(MoveTempIfPossible(InParams))
+			.Then([Promise = MoveTemp(Promise), Op = InAsyncOp.AsShared(), ResultKey](TFuture<TOnlineResult<SecondaryOpType>>&& Future) mutable
+			{
+				if (Future.Get().IsError())
+				{
+					UE_LOG(LogLobbies, Log, TEXT("[FLobbiesCommon::CaptureStepResult] Failed secondary operation %s. Error: %s"), SecondaryOpType::Name, *Future.Get().GetErrorValue().GetLogString());
+					Op->SetError(Errors::RequestFailure(MoveTempIfPossible(Future.Get().GetErrorValue())));
+				}
+				else
+				{
+					UE_LOG(LogLobbies, Log, TEXT("[FLobbiesCommon::CaptureStepResult] Captured secondary operation %s as %s"), SecondaryOpType::Name, *ResultKey);
+					Op->Data.Set(ResultKey, MoveTempIfPossible(Future.Get().GetOkValue()));
+				}
+				Promise.EmplaceValue();
+			});
+
+			return Future;
+		};
+	}
+
+	template <typename SecondaryOpType, typename OpType, typename Function>
+	TFunction<TFuture<void>(TOnlineAsyncOp<OpType>&)> CaptureStepResult(TOnlineAsyncOp<OpType>& InAsyncOp, const FString& ResultKey, Function Func, typename SecondaryOpType::Params&& InParams)
+	{
+		return [this, ResultKey, Func, InParams = MoveTempIfPossible(InParams)](TOnlineAsyncOp<OpType>& InAsyncOp) mutable
+		{
+			TPromise<void> Promise;
+			auto Future = Promise.GetFuture();
+
+			UE_LOG(LogLobbies, Log, TEXT("[FLobbiesCommon::CaptureStepResult] Starting secondary operation %s."), SecondaryOpType::Name);
+
+			(this->*Func)(MoveTempIfPossible(InParams))
+			.Then([Promise = MoveTemp(Promise), Op = InAsyncOp.AsShared(), ResultKey](TFuture<TOnlineResult<SecondaryOpType>>&& Future) mutable
+			{
+				if (Future.Get().IsError())
+				{
+					UE_LOG(LogLobbies, Log, TEXT("[FLobbiesCommon::CaptureStepResult] Failed secondary operation %s. Error: %s"), SecondaryOpType::Name, *Future.Get().GetErrorValue().GetLogString());
+					Op->SetError(Errors::RequestFailure(MoveTempIfPossible(Future.Get().GetErrorValue())));
+				}
+				else
+				{
+					UE_LOG(LogLobbies, Log, TEXT("[FLobbiesCommon::CaptureStepResult] Captured secondary operation %s as %s"), SecondaryOpType::Name, *ResultKey);
+					Op->Data.Set(ResultKey, MoveTempIfPossible(Future.Get().GetOkValue()));
+				}
+				Promise.EmplaceValue();
+			});
+
+			return Future;
+		};
+	}
+
+	template <typename SecondaryOpType, typename OpType, typename Function>
+	TFunction<TFuture<void>(TOnlineAsyncOp<OpType>&, typename SecondaryOpType::Params&&)> ConsumeOperationStepResult(TOnlineAsyncOp<OpType>& InAsyncOp, Function Func)
+	{
+		return [this, Func](TOnlineAsyncOp<OpType>& InAsyncOp, typename SecondaryOpType::Params&& InParams) mutable -> TFuture<void>
+		{
+			auto Promise = MakeShared<TPromise<void>>();
+			auto Future = Promise->GetFuture();
+
+			UE_LOG(LogLobbies, Log, TEXT("[FLobbiesCommon::ConsumeOperationStepResult] Starting secondary operation %s."), SecondaryOpType::Name);
+
+			(this->*Func)(MoveTempIfPossible(InParams))
+			.OnComplete([Promise, Op = InAsyncOp.AsShared()](const TOnlineResult<SecondaryOpType>& Result) mutable -> void
+			{
+				if (Result.IsError())
+				{
+					UE_LOG(LogLobbies, Log, TEXT("[FLobbiesCommon::ConsumeOperationStepResult] Failed secondary operation %s. Error: %s"), SecondaryOpType::Name, *Result.GetErrorValue().GetLogString());
+					Op->SetError(Errors::RequestFailure(MoveTempIfPossible(Result.GetErrorValue())));
+				}
+				else
+				{
+					UE_LOG(LogLobbies, Log, TEXT("[FLobbiesCommon::ConsumeOperationStepResult] Completed secondary operation %s."), SecondaryOpType::Name);
+				}
+				Promise->EmplaceValue();
+			});
+
+			return Future;
+		};
+	}
+
+	template <typename SecondaryOpType, typename OpType, typename Function>
+	TFunction<TFuture<void>(TOnlineAsyncOp<OpType>&)> ConsumeOperationStepResult(TOnlineAsyncOp<OpType>& InAsyncOp, Function Func, typename SecondaryOpType::Params&& InParams)
+	{
+		return [this, Func, InParams = MoveTempIfPossible(InParams)](TOnlineAsyncOp<OpType>& InAsyncOp) mutable
+		{
+			auto Promise = MakeShared<TPromise<void>>();
+			auto Future = Promise->GetFuture();
+
+			UE_LOG(LogLobbies, Log, TEXT("[FLobbiesCommon::ConsumeOperationStepResult] Starting secondary operation %s."), SecondaryOpType::Name);
+
+			(this->*Func)(MoveTempIfPossible(InParams))
+			.OnComplete([Promise, Op = InAsyncOp.AsShared()](const TOnlineResult<SecondaryOpType>& Result) mutable
+			{
+				if (Result.IsError())
+				{
+					UE_LOG(LogLobbies, Log, TEXT("[FLobbiesCommon::ConsumeOperationStepResult] Failed secondary operation %s. Error: %s"), SecondaryOpType::Name, *Result.GetErrorValue().GetLogString());
+					Op->SetError(Errors::RequestFailure(MoveTempIfPossible(Result.GetErrorValue())));
+				}
+				else
+				{
+					UE_LOG(LogLobbies, Log, TEXT("[FLobbiesCommon::ConsumeOperationStepResult] Completed secondary operation %s."), SecondaryOpType::Name);
+				}
+				Promise->EmplaceValue();
+			});
+
+			return Future;
+		};
+	}
+
+	template <typename SecondaryOpType, typename OpType, typename Function>
+	TFunction<TFuture<void>(TOnlineAsyncOp<OpType>&, typename SecondaryOpType::Params&&)> CaptureOperationStepResult(TOnlineAsyncOp<OpType>& InAsyncOp, const FString& ResultKey, Function Func)
+	{
+		return [this, ResultKey, Func](TOnlineAsyncOp<OpType>& InAsyncOp, typename SecondaryOpType::Params&& InParams) mutable -> TFuture<void>
+		{
+			auto Promise = MakeShared<TPromise<void>>();
+			auto Future = Promise->GetFuture();
+
+			UE_LOG(LogLobbies, Log, TEXT("[FLobbiesCommon::CaptureOperationStepResult] Starting secondary operation %s."), SecondaryOpType::Name);
+
+			(this->*Func)(MoveTempIfPossible(InParams))
+			.OnComplete([Promise, Op = InAsyncOp.AsShared(), ResultKey](const TOnlineResult<SecondaryOpType>& Result) mutable -> void
+			{
+				if (Result.IsError())
+				{
+					UE_LOG(LogLobbies, Log, TEXT("[FLobbiesCommon::CaptureOperationStepResult] Failed secondary operation %s. Error: %s"), SecondaryOpType::Name, *Result.GetErrorValue().GetLogString());
+					Op->SetError(Errors::RequestFailure(MoveTempIfPossible(Result.GetErrorValue())));
+				}
+				else
+				{
+					UE_LOG(LogLobbies, Log, TEXT("[FLobbiesCommon::CaptureOperationStepResult] Captured secondary operation %s as %s"), SecondaryOpType::Name, *ResultKey);
+					Op->Data.Set(ResultKey, Result.GetOkValue());
+				}
+				Promise->EmplaceValue();
+			});
+
+			return Future;
+		};
+	}
+
+	template <typename SecondaryOpType, typename OpType, typename Function>
+	TFunction<TFuture<void>(TOnlineAsyncOp<OpType>&)> CaptureOperationStepResult(TOnlineAsyncOp<OpType>& InAsyncOp, const FString& ResultKey, Function Func, typename SecondaryOpType::Params&& InParams)
+	{
+		return [this, ResultKey, Func, InParams = MoveTempIfPossible(InParams)](TOnlineAsyncOp<OpType>& InAsyncOp) mutable
+		{
+			auto Promise = MakeShared<TPromise<void>>();
+			auto Future = Promise->GetFuture();
+
+			UE_LOG(LogLobbies, Log, TEXT("[FLobbiesCommon::CaptureOperationStepResult] Starting secondary operation %s."), SecondaryOpType::Name);
+
+			(this->*Func)(MoveTempIfPossible(InParams))
+			.OnComplete([Promise, Op = InAsyncOp.AsShared(), ResultKey](const TOnlineResult<SecondaryOpType>& Result) mutable
+			{
+				if (Result.IsError())
+				{
+					UE_LOG(LogLobbies, Log, TEXT("[FLobbiesCommon::CaptureOperationStepResult] Failed secondary operation %s. Error: %s"), SecondaryOpType::Name, *Result.GetErrorValue().GetLogString());
+					Op->SetError(Errors::RequestFailure(MoveTempIfPossible(Result.GetErrorValue())));
+				}
+				else
+				{
+					UE_LOG(LogLobbies, Log, TEXT("[FLobbiesCommon::CaptureOperationStepResult] Captured secondary operation %s as %s"), SecondaryOpType::Name, *ResultKey);
+					Op->Data.Set(ResultKey, Result.GetOkValue());
+				}
+				Promise->EmplaceValue();
+			});
+
+			return Future;
+		};
+	}
+
+	TOnlineEventCallable<void(const FLobbyJoined&)> OnLobbyJoinedEvent;
+	TOnlineEventCallable<void(const FLobbyLeft&)> OnLobbyLeftEvent;
+	TOnlineEventCallable<void(const FLobbyMemberJoined&)> OnLobbyMemberJoinedEvent;
+	TOnlineEventCallable<void(const FLobbyMemberLeft&)> OnLobbyMemberLeftEvent;
+	TOnlineEventCallable<void(const FLobbyLeaderChanged&)> OnLobbyLeaderChangedEvent;
+	TOnlineEventCallable<void(const FLobbySchemaChanged&)> OnLobbySchemaChangedEvent;
+	TOnlineEventCallable<void(const FLobbyAttributesChanged&)> OnLobbyAttributesChangedEvent;
+	TOnlineEventCallable<void(const FLobbyMemberAttributesChanged&)> OnLobbyMemberAttributesChangedEvent;
+	TOnlineEventCallable<void(const FLobbyInvitationAdded&)> OnLobbyInvitationAddedEvent;
+	TOnlineEventCallable<void(const FLobbyInvitationRemoved&)> OnLobbyInvitationRemovedEvent;
+
+	TSharedPtr<FLobbySchemaRegistry> LobbySchemaRegistry;
+	TSharedPtr<FLobbySchema> ServiceSchema;
+};
+
+namespace Meta {
+
+BEGIN_ONLINE_STRUCT_META(FLobbyConfig)
+	ONLINE_STRUCT_FIELD(FLobbyConfig, RegisteredSchemas)
+END_ONLINE_STRUCT_META()
+
+
+BEGIN_ONLINE_STRUCT_META(FLobbySchemaAttributeConfig)
+	ONLINE_STRUCT_FIELD(FLobbySchemaAttributeConfig, Name),
+	ONLINE_STRUCT_FIELD(FLobbySchemaAttributeConfig, ServiceAttributeFieldId),
+	ONLINE_STRUCT_FIELD(FLobbySchemaAttributeConfig, MaxByteSize),
+	ONLINE_STRUCT_FIELD(FLobbySchemaAttributeConfig, Visibility)
+END_ONLINE_STRUCT_META()
+
+BEGIN_ONLINE_STRUCT_META(FLobbySchemaConfig)
+	ONLINE_STRUCT_FIELD(FLobbySchemaConfig, BaseSchemaName),
+	ONLINE_STRUCT_FIELD(FLobbySchemaConfig, LobbyAttributes),
+	ONLINE_STRUCT_FIELD(FLobbySchemaConfig, LobbyMemberAttributes)
+END_ONLINE_STRUCT_META()
+
+#if LOBBIES_FUNCTIONAL_TEST_ENABLED
+	BEGIN_ONLINE_STRUCT_META(FFunctionalTestLobbies::Params)
+	END_ONLINE_STRUCT_META()
+
+	BEGIN_ONLINE_STRUCT_META(FFunctionalTestLobbies::Result)
+	END_ONLINE_STRUCT_META()
+#endif
+
+/* Meta */ }
+
+/* UE::Online */ }
diff --git a/Engine/Plugins/Online/Experimental/OnlineServices/Source/OnlineServicesInterface/Public/Online/Lobbies.cpp b/Engine/Plugins/Online/Experimental/OnlineServices/Source/OnlineServicesInterface/Public/Online/Lobbies.cpp
new file mode 100644
--- /dev/null
+++ b/Engine/Plugins/Online/Experimental/OnlineServices/Source/OnlineServicesInterface/Public/Online/Lobbies.cpp
@@ -0,0 +1,102 @@
+// Copyright Epic Games, Inc. All Rights Reserved.
+
+#include "Online/Lobbies.h"
+
+namespace UE::Online {
+
+const TCHAR* LexToString(ELobbyJoinPolicy Policy)
+{
+	switch (Policy)
+	{
+	case ELobbyJoinPolicy::PublicAdvertised:	return TEXT("PublicAdvertised");
+	case ELobbyJoinPolicy::PublicNotAdvertised:	return TEXT("PublicNotAdvertised");
+	default:									checkNoEntry(); // Intentional fallthrough
+	case ELobbyJoinPolicy::InvitationOnly:		return TEXT("InvitationOnly");
+	}
+}
+
+void LexFromString(ELobbyJoinPolicy& OutPolicy, const TCHAR* InStr)
+{
+	if (FCString::Stricmp(InStr, TEXT("PublicAdvertised")) == 0)
+	{
+		OutPolicy = ELobbyJoinPolicy::PublicAdvertised;
+	}
+	else if (FCString::Stricmp(InStr, TEXT("PublicNotAdvertised")) == 0)
+	{
+		OutPolicy = ELobbyJoinPolicy::PublicNotAdvertised;
+	}
+	else if (FCString::Stricmp(InStr, TEXT("InvitationOnly")) == 0)
+	{
+		OutPolicy = ELobbyJoinPolicy::InvitationOnly;
+	}
+	else
+	{
+		checkNoEntry();
+		OutPolicy = ELobbyJoinPolicy::InvitationOnly;
+	}
+}
+
+const TCHAR* LexToString(ELobbyComparisonOp Comparison)
+{
+	switch (Comparison)
+	{
+	case ELobbyComparisonOp::Equals:			return TEXT("Equals");
+	case ELobbyComparisonOp::NotEquals:			return TEXT("NotEquals");
+	case ELobbyComparisonOp::GreaterThan:		return TEXT("GreaterThan");
+	case ELobbyComparisonOp::GreaterThanEquals:	return TEXT("GreaterThanEquals");
+	case ELobbyComparisonOp::LessThan:			return TEXT("LessThan");
+	case ELobbyComparisonOp::LessThanEquals:	return TEXT("LessThanEquals");
+	case ELobbyComparisonOp::Near:				return TEXT("Near");
+	case ELobbyComparisonOp::In:				return TEXT("In");
+	default:									checkNoEntry(); // Intentional fallthrough
+	case ELobbyComparisonOp::NotIn:				return TEXT("NotIn");
+	}
+}
+
+void LexFromString(ELobbyComparisonOp& OutComparison, const TCHAR* InStr)
+{
+	if (FCString::Stricmp(InStr, TEXT("Equals")) == 0)
+	{
+		OutComparison = ELobbyComparisonOp::Equals;
+	}
+	else if (FCString::Stricmp(InStr, TEXT("NotEquals")) == 0)
+	{
+		OutComparison = ELobbyComparisonOp::NotEquals;
+	}
+	else if (FCString::Stricmp(InStr, TEXT("GreaterThan")) == 0)
+	{
+		OutComparison = ELobbyComparisonOp::GreaterThan;
+	}
+	else if (FCString::Stricmp(InStr, TEXT("GreaterThanEquals")) == 0)
+	{
+		OutComparison = ELobbyComparisonOp::GreaterThanEquals;
+	}
+	else if (FCString::Stricmp(InStr, TEXT("LessThan")) == 0)
+	{
+		OutComparison = ELobbyComparisonOp::LessThan;
+	}
+	else if (FCString::Stricmp(InStr, TEXT("LessThanEquals")) == 0)
+	{
+		OutComparison = ELobbyComparisonOp::LessThanEquals;
+	}
+	else if (FCString::Stricmp(InStr, TEXT("Near")) == 0)
+	{
+		OutComparison = ELobbyComparisonOp::Near;
+	}
+	else if (FCString::Stricmp(InStr, TEXT("In")) == 0)
+	{
+		OutComparison = ELobbyComparisonOp::In;
+	}
+	else
+	{
+		checkNoEntry();
+		OutComparison = ELobbyComparisonOp::In;
+	}
+}
+
+void SortLobbies(const TArray<FFindLobbySearchFilter>& Filters, TArray<TSharedRef<const FLobby>>& Lobbies)
+{
+	// todo
+}
+
+/* UE::Online */ }
diff --git a/Engine/Plugins/Online/Experimental/OnlineServices/Source/OnlineServicesInterface/Public/Online/Lobbies.h b/Engine/Plugins/Online/Experimental/OnlineServices/Source/OnlineServicesInterface/Public/Online/Lobbies.h
new file mode 100644
--- /dev/null
+++ b/Engine/Plugins/Online/Experimental/OnlineServices/Source/OnlineServicesInterface/Public/Online/Lobbies.h
@@ -0,0 +1,865 @@
+// Copyright Epic Games, Inc. All Rights Reserved.
+
+#pragma once
+
+#include "Online/OnlineAsyncOpHandle.h"
+#include "Online/CoreOnline.h"
+#include "Online/OnlineMeta.h"
+
+namespace UE::Online {
+
+typedef FName FLobbySchemaId;
+typedef FName FLobbyAttributeId;
+using FLobbyVariant = FString; // TODO:  Import FVariantData
+
+enum class ELobbyJoinPolicy : uint8
+{
+	/** 
+	* Lobby can be found through searches based on attribute matching,
+	*  by knowing the lobby id, or by invitation.
+	*/
+	PublicAdvertised,
+
+	/**  Lobby may be joined by knowing the lobby id or by invitation. */
+	PublicNotAdvertised,
+
+	/**  Lobby may only be joined by invitation. */
+	InvitationOnly,
+};
+ONLINESERVICESINTERFACE_API const TCHAR* LexToString(ELobbyJoinPolicy Policy);
+ONLINESERVICESINTERFACE_API void LexFromString(ELobbyJoinPolicy& OutPolicy, const TCHAR* InStr);
+
+enum class ELobbyAttributeVisibility
+{
+	Public,
+	Private,
+};
+
+struct FLobbyMember
+{
+	FOnlineAccountIdHandle AccountId;
+	FOnlineAccountIdHandle PlatformAccountId;
+	FString PlatformDisplayName;
+	TMap<FLobbyAttributeId, FLobbyVariant> Attributes;
+};
+
+struct FLobby
+{
+	FOnlineLobbyIdHandle LobbyId;
+	FOnlineAccountIdHandle OwnerAccountId;
+	FName LocalName;
+	FName SchemaName;
+	int32 MaxMembers;
+	ELobbyJoinPolicy JoinPolicy;
+	TMap<FLobbyAttributeId, FLobbyVariant> Attributes;
+	TArray<TSharedRef<const FLobbyMember>> Members;
+};
+
+struct FJoinLobbyLocalUserData
+{
+	/** Local users who will be joining the lobby. */
+	FOnlineAccountIdHandle LocalUserId;
+
+	/** Initial attributes. */
+	TMap<FLobbyAttributeId, FLobbyVariant> Attributes;
+};
+
+enum class ELobbyComparisonOp : uint8
+{
+
+	Equals,
+	NotEquals,
+	GreaterThan,
+	GreaterThanEquals,
+	LessThan,
+	LessThanEquals,
+	Near,
+	In,
+	NotIn
+};
+
+ONLINESERVICESINTERFACE_API const TCHAR* LexToString(ELobbyComparisonOp EnumVal);
+ONLINESERVICESINTERFACE_API void LexFromString(ELobbyComparisonOp& OutComparison, const TCHAR* InStr);
+
+struct FFindLobbySearchFilter
+{
+	/** The name of the attribute to be compared. */
+	FLobbyAttributeId AttributeName;
+
+	/** The type of comparison to perform. */
+	ELobbyComparisonOp ComparisonOp;
+
+	/** Value to use when comparing the attribute. */
+	FLobbyVariant ComparisonValue;
+};
+
+ONLINESERVICESINTERFACE_API void SortLobbies(const TArray<FFindLobbySearchFilter>& Filters, TArray<TSharedRef<const FLobby>>& Lobbies);
+
+struct FCreateLobby
+{
+	static constexpr TCHAR Name[] = TEXT("CreateLobby");
+
+	/** Input struct for Lobbies::CreateLobby */
+	struct Params
+	{
+		/** The local user agent which will perform the action. */
+		FOnlineAccountIdHandle LocalUserId;
+
+		/** The local name for the lobby. */
+		FName LocalName;
+
+		/** The schema which will be applied to the lobby. */
+		FName SchemaName;
+
+		/** The maximum number of members who can fit in the lobby. */
+		int32 MaxMembers = 0;
+
+		/** Initial join policy setting. */
+		ELobbyJoinPolicy JoinPolicy = ELobbyJoinPolicy::InvitationOnly;
+
+		/** Initial attributes. */
+		TMap<FLobbyAttributeId, FLobbyVariant> Attributes;
+
+		/** Local users who will be joining the lobby. */
+		TArray<FJoinLobbyLocalUserData> LocalUsers;
+	};
+
+	/** Output struct for Lobbies::CreateLobby */
+	struct Result
+	{
+		// Todo: investigate returning TSharedRef
+		TSharedPtr<const FLobby> Lobby;
+	};
+};
+
+struct FFindLobby
+{
+	static constexpr TCHAR Name[] = TEXT("FindLobby");
+
+	/** Input struct for Lobbies::FindLobby */
+	struct Params
+	{
+		/** The local user agent which will perform the action. */
+		FOnlineAccountIdHandle LocalUserId;
+
+		/** 
+		* Max results to return in one search.
+		*  Actual count may be smaller based on implementation.
+		*/
+		uint32 MaxResults = 20;
+
+		/**  Filters to apply when searching for lobbies. */
+		TArray<FFindLobbySearchFilter> Filters;
+
+		/**  Find lobbies containing the target user. */
+		TOptional<FOnlineAccountIdHandle> TargetUser;
+
+		/**  Find join info for the target lobby id. */
+		TOptional<FOnlineLobbyIdHandle> LobbyId;
+	};
+
+	/** Output struct for Lobbies::FindLobby */
+	struct Result
+	{
+		TArray<TSharedRef<const FLobby>> Lobbies;
+	};
+};
+
+struct FRestoreLobbies
+{
+	static constexpr TCHAR Name[] = TEXT("RestoreLobbies");
+
+	/** Input struct for Lobbies::RestoreLobbies */
+	struct Params
+	{
+	};
+
+	/** Output struct for Lobbies::RestoreLobbies */
+	struct Result
+	{
+	};
+};
+
+struct FJoinLobby
+{
+	static constexpr TCHAR Name[] = TEXT("JoinLobby");
+
+	/** Input struct for Lobbies::JoinLobby */
+	struct Params
+	{
+		/**  The local user agent which will perform the action. */
+		FOnlineAccountIdHandle LocalUserId;
+
+		/**  The local name for the lobby. */
+		FName LocalName;
+
+		/**  The id of the lobby to be joined. */
+		FOnlineLobbyIdHandle LobbyId;
+
+		/**  Local users who will be joining the lobby. */
+		TArray<FJoinLobbyLocalUserData> LocalUsers;
+	};
+
+	/** Output struct for Lobbies::JoinLobby */
+	struct Result
+	{
+		// Todo: investigate returning TSharedRef
+		TSharedPtr<const FLobby> Lobby;
+	};
+};
+
+struct FLeaveLobby
+{
+	static constexpr TCHAR Name[] = TEXT("LeaveLobby");
+
+	/** Input struct for Lobbies::LeaveLobby */
+	struct Params
+	{
+		/**  The local user agent which will perform the action. */
+		FOnlineAccountIdHandle LocalUserId;
+
+		/**  Id of the lobby to leave. */
+		FOnlineLobbyIdHandle LobbyId;
+	};
+
+	/** Output struct for Lobbies::LeaveLobby */
+	struct Result
+	{
+	};
+};
+
+struct FInviteLobbyMember
+{
+	static constexpr TCHAR Name[] = TEXT("InviteLobbyMember");
+
+	/** Input struct for Lobbies::InviteLobbyMember */
+	struct Params
+	{
+		/**  The local user agent which will perform the action. */
+		FOnlineAccountIdHandle LocalUserId;
+
+		/**  Id of the lobby for which the invitation will be sent. */
+		FOnlineLobbyIdHandle LobbyId;
+
+		/**  Id of the player who will be sent the invitation. */
+		FOnlineAccountIdHandle TargetUserId;
+	};
+
+	/** Output struct for Lobbies::InviteLobbyMember */
+	struct Result
+	{
+	};
+};
+
+struct FDeclineLobbyInvitation
+{
+	static constexpr TCHAR Name[] = TEXT("DeclineLobbyInvitation");
+
+	/** Input struct for Lobbies::DeclineLobbyInvitation */
+	struct Params
+	{
+		/**  The local user agent which will perform the action. */
+		FOnlineAccountIdHandle LocalUserId;
+
+		// Id of the lobby for which the invitations will be declined.
+		FOnlineLobbyIdHandle LobbyId;
+	};
+
+	/** Output struct for Lobbies::DeclineLobbyInvitation */
+	struct Result
+	{
+	};
+};
+
+struct FKickLobbyMember
+{
+	static constexpr TCHAR Name[] = TEXT("KickLobbyMember");
+
+	/** Input struct for Lobbies::KickLobbyMember */
+	struct Params
+	{
+		/** The local user agent which will perform the action. */
+		FOnlineAccountIdHandle LocalUserId;
+
+		/** Id of the lobby. */
+		FOnlineLobbyIdHandle LobbyId;
+
+		/** The target user to be kicked. */
+		FOnlineAccountIdHandle TargetUserId;
+	};
+
+	/** Output struct for Lobbies::KickLobbyMember */
+	struct Result
+	{
+	};
+};
+
+struct FPromoteLobbyMember
+{
+	static constexpr TCHAR Name[] = TEXT("PromoteLobbyMember");
+
+	/** Input struct for Lobbies::PromoteLobbyMember */
+	struct Params
+	{
+		/** The local user agent which will perform the action. */
+		FOnlineAccountIdHandle LocalUserId;
+
+		/** Id of the lobby. */
+		FOnlineLobbyIdHandle LobbyId;
+
+		/** The target user to be promoted to owner. */
+		FOnlineAccountIdHandle TargetUserId;
+	};
+
+	/** Output struct for Lobbies::PromoteLobbyMember */
+	struct Result
+	{
+	};
+};
+
+struct FModifyLobbySchema
+{
+	static constexpr TCHAR Name[] = TEXT("ModifyLobbySchema");
+
+	/** Input struct for Lobbies::ModifyLobbySchema */
+	struct Params
+	{
+	};
+
+	/** Output struct for Lobbies::ModifyLobbySchema */
+	struct Result
+	{
+	};
+};
+
+struct FModifyLobbyJoinPolicy
+{
+	static constexpr TCHAR Name[] = TEXT("ModifyLobbyJoinPolicy");
+
+	/** Input struct for Lobbies::ModifyLobbyJoinPolicy */
+	struct Params
+	{
+		/** The local user agent which will perform the action. */
+		FOnlineAccountIdHandle LocalUserId;
+
+		/** Id of the lobby. */
+		FOnlineLobbyIdHandle LobbyId;
+
+		/** The new join policy setting. */
+		ELobbyJoinPolicy JoinPolicy = ELobbyJoinPolicy::InvitationOnly;
+	};
+
+	/** Output struct for Lobbies::ModifyLobbyJoinPolicy */
+	struct Result
+	{
+	};
+};
+
+struct FModifyLobbyAttributes
+{
+	static constexpr TCHAR Name[] = TEXT("ModifyLobbyAttributes");
+
+	/** Input struct for Lobbies::ModifyLobbyAttributes */
+	struct Params
+	{
+		/** The local user agent which will perform the action. */
+		FOnlineAccountIdHandle LocalUserId;
+
+		/** Id of the lobby. */
+		FOnlineLobbyIdHandle LobbyId;
+
+		/** New or changed lobby attributes. */
+		TMap<FLobbyAttributeId, FLobbyVariant> MutatedAttributes;
+
+		/** Attributes to be cleared. */
+		TSet<FLobbyAttributeId> ClearedAttributes;
+	};
+
+	/** Output struct for Lobbies::ModifyLobbyAttributes */
+	struct Result
+	{
+	};
+};
+
+struct FModifyLobbyMemberAttributes
+{
+	static constexpr TCHAR Name[] = TEXT("ModifyLobbyMemberAttributes");
+
+	/** Input struct for Lobbies::ModifyLobbyMemberAttributes */
+	struct Params
+	{
+		/** The local user agent which will perform the action. */
+		FOnlineAccountIdHandle LocalUserId;
+
+		/** Id of the lobby. */
+		FOnlineLobbyIdHandle LobbyId;
+
+		/** New or changed lobby attributes. */
+		TMap<FLobbyAttributeId, FLobbyVariant> MutatedAttributes;
+
+		/** Attributes to be cleared. */
+		TSet<FLobbyAttributeId> ClearedAttributes;
+	};
+
+	/** Output struct for Lobbies::ModifyLobbyMemberAttributes */
+	struct Result
+	{
+	};
+};
+
+struct FGetJoinedLobbies
+{
+	static constexpr TCHAR Name[] = TEXT("GetJoinedLobbies");
+
+	/** Input struct for Lobbies::GetJoinedLobbies */
+	struct Params
+	{
+		/** The local user agent which will perform the action. */
+		FOnlineAccountIdHandle LocalUserId;
+	};
+
+	/** Output struct for Lobbies::GetJoinedLobbies */
+	struct Result
+	{
+		TArray<TSharedRef<FLobby>> Lobbies;
+	};
+};
+
+struct FGetReceivedInvitations
+{
+	static constexpr TCHAR Name[] = TEXT("GetReceivedInvitations");
+
+	/** Input struct for Lobbies::GetReceivedInvitations */
+	struct Params
+	{
+		/** The local user agent which will perform the action. */
+		FOnlineAccountIdHandle LocalUserId;
+	};
+
+	/** Output struct for Lobbies::GetReceivedInvitations */
+	struct Result
+	{
+	};
+};
+
+/** Struct for LobbyJoined event */
+struct FLobbyJoined
+{
+	/** Lobby data access. */
+	TSharedRef<const FLobby> Lobby;
+
+};
+
+/** Struct for LobbyLeft event */
+struct FLobbyLeft
+{
+	/** Lobby data access. */
+	TSharedRef<const FLobby> Lobby;
+
+};
+
+/** Struct for LobbyMemberJoined event */
+struct FLobbyMemberJoined
+{
+	/** Lobby data access. */
+	TSharedRef<const FLobby> Lobby;
+
+	/** Member data access. */
+	TSharedRef<const FLobbyMember> Member;
+
+};
+
+/** Struct for LobbyMemberLeft event */
+struct FLobbyMemberLeft
+{
+	/** Lobby data access. */
+	TSharedRef<const FLobby> Lobby;
+
+	/** Member data access. */
+	TSharedRef<const FLobbyMember> Member;
+
+};
+
+/** Struct for LobbyLeaderChanged event */
+struct FLobbyLeaderChanged
+{
+};
+
+/** Struct for LobbySchemaChanged event */
+struct FLobbySchemaChanged
+{
+};
+
+/** Struct for LobbyAttributesChanged event */
+struct FLobbyAttributesChanged
+{
+};
+
+/** Struct for LobbyMemberAttributesChanged event */
+struct FLobbyMemberAttributesChanged
+{
+};
+
+/** Struct for LobbyInvitationAdded event */
+struct FLobbyInvitationAdded
+{
+	/** The local user associated with the invitation. */
+	FOnlineAccountIdHandle LocalUserId;
+
+	/** The user who sent the invitation. */
+	FOnlineAccountIdHandle SenderId;
+
+	/** The invited lobby. */
+	TSharedRef<const FLobby> Lobby;
+};
+
+/** Struct for LobbyInvitationRemoved event */
+struct FLobbyInvitationRemoved
+{
+	/** The local user associated with the invitation. */
+	FOnlineAccountIdHandle LocalUserId;
+
+	/** The user who sent the invitation. */
+	FOnlineAccountIdHandle SenderId;
+
+	/** The invited lobby. */
+	TSharedRef<const FLobby> Lobby;
+
+};
+
+class ILobbies
+{
+public:
+	//----------------------------------------------------------------------------------------------
+	// Operations
+
+	/**
+	 * Create and join a new lobby.
+	 *
+	 * @param Params for the CreateLobby call
+	 * @return
+	 */
+	virtual TOnlineAsyncOpHandle<FCreateLobby> CreateLobby(FCreateLobby::Params&& Params) = 0;
+
+	/**
+	 * Search for lobbies using filtering parameters.
+	 *
+	 * @param Params for the FindLobby call
+	 * @return
+	 */
+	virtual TOnlineAsyncOpHandle<FFindLobby> FindLobby(FFindLobby::Params&& Params) = 0;
+
+	/**
+	 * Try to rejoin previously joined lobbies.
+	 *
+	 * @param Params for the RestoreLobbies call
+	 * @return
+	 */
+	virtual TOnlineAsyncOpHandle<FRestoreLobbies> RestoreLobbies(FRestoreLobbies::Params&& Params) = 0;
+
+	/**
+	 * Join a lobby using its id.
+	 *
+	 * @param Params for the JoinLobby call
+	 * @return
+	 */
+	virtual TOnlineAsyncOpHandle<FJoinLobby> JoinLobby(FJoinLobby::Params&& Params) = 0;
+
+	/**
+	 * Leave a joined lobby.
+	 *
+	 * @param Params for the LeaveLobby call
+	 * @return
+	 */
+	virtual TOnlineAsyncOpHandle<FLeaveLobby> LeaveLobby(FLeaveLobby::Params&& Params) = 0;
+
+	/**
+	 * Invite a player to join a lobby.
+	 *
+	 * @param Params for the InviteLobbyMember call
+	 * @return
+	 */
+	virtual TOnlineAsyncOpHandle<FInviteLobbyMember> InviteLobbyMember(FInviteLobbyMember::Params&& Params) = 0;
+
+	/**
+	 * Decline an invitation to join a lobby.
+	 *
+	 * @param Params for the DeclineLobbyInvitation call
+	 * @return
+	 */
+	virtual TOnlineAsyncOpHandle<FDeclineLobbyInvitation> DeclineLobbyInvitation(FDeclineLobbyInvitation::Params&& Params) = 0;
+
+	/**
+	 * Kick a member from a the target lobby.
+	 *
+	 * @param Params for the KickLobbyMember call
+	 * @return
+	 */
+	virtual TOnlineAsyncOpHandle<FKickLobbyMember> KickLobbyMember(FKickLobbyMember::Params&& Params) = 0;
+
+	/**
+	 * Promote another lobby member to leader.
+	 * The local user must be the current lobby leader to promote another member.
+	 *
+	 * @param Params for the PromoteLobbyMember call
+	 * @return
+	 */
+	virtual TOnlineAsyncOpHandle<FPromoteLobbyMember> PromoteLobbyMember(FPromoteLobbyMember::Params&& Params) = 0;
+
+	//----------------------------------------------------------------------------------------------
+	// Mutations
+
+	/**
+	 * Change the schema applied to the lobby and member attributes.
+	 * Only the lobby leader may change the schema. Existing attributes not present in the new
+	 * schema will be cleared.
+	 *
+	 * @param Params for the ModifyLobbySchema call
+	 * @return
+	 */
+	virtual TOnlineAsyncOpHandle<FModifyLobbySchema> ModifyLobbySchema(FModifyLobbySchema::Params&& Params) = 0;
+
+	/**
+	 * Change the join policy applied to the lobby.
+	 * Only the lobby leader may change the join policy.
+	 *
+	 * @param Params for the ModifyLobbyJoinPolicy call
+	 * @return
+	 */
+	virtual TOnlineAsyncOpHandle<FModifyLobbyJoinPolicy> ModifyLobbyJoinPolicy(FModifyLobbyJoinPolicy::Params&& Params) = 0;
+
+	/**
+	 * Change the attributes applied to the lobby.
+	 * Only the lobby leader may change the lobby attributes.
+	 * Attributes are validated against the lobby schema before an update can succeed.
+	 *
+	 * @param Params for the ModifyLobbyAttributes call
+	 * @return
+	 */
+	virtual TOnlineAsyncOpHandle<FModifyLobbyAttributes> ModifyLobbyAttributes(FModifyLobbyAttributes::Params&& Params) = 0;
+
+	/**
+	 * Change the attributes applied to a lobby member.
+	 * Lobby members may only change their own attributes.
+	 * Attributes are validated against the lobby schema before an update can succeed.
+	 *
+	 * @param Params for the ModifyLobbyAttributes call
+	 * @return
+	 */
+	virtual TOnlineAsyncOpHandle<FModifyLobbyMemberAttributes> ModifyLobbyMemberAttributes(FModifyLobbyMemberAttributes::Params&& Params) = 0;
+
+	//----------------------------------------------------------------------------------------------
+	// Accessors
+
+	/**
+	 * Retrieve the list of joined lobbies for the target local user.
+	 *
+	 * @param Params for the GetJoinedLobbies call
+	 * @return
+	 */
+	virtual TOnlineResult<FGetJoinedLobbies> GetJoinedLobbies(FGetJoinedLobbies::Params&& Params) = 0;
+
+	/**
+	 * Retrieve the list of received invitations for the target local user.
+	 *
+	 * @param Params for the GetJoinedLobbies call
+	 * @return
+	 */
+	virtual TOnlineResult<FGetReceivedInvitations> GetReceivedInvitations(FGetReceivedInvitations::Params&& Params) = 0;
+
+	//----------------------------------------------------------------------------------------------
+	// Events
+
+	/**
+	 * Get the event that is triggered when a lobby is joined.
+	 * This will happen as a result of creating or joining a lobby.
+	 *
+	 * @return Event that can be bound to
+	 */
+	virtual TOnlineEvent<void(const FLobbyJoined&)> OnLobbyJoined() = 0;
+
+	/**
+	 * Get the event that is triggered when a lobby has been left by all local members.
+	 *
+	 * @return Event that can be bound to
+	 */
+	virtual TOnlineEvent<void(const FLobbyLeft&)> OnLobbyLeft() = 0;
+
+	/**
+	 * Get the event that is triggered when a lobby member joins.
+	 * This will happen as a result of creating or joining a lobby for local users, and will trigger
+	 * when a remote user joins.
+	 *
+	 * @return Event that can be bound to
+	 */
+	virtual TOnlineEvent<void(const FLobbyMemberJoined&)> OnLobbyMemberJoined() = 0;
+
+	/**
+	 * Get the event that is triggered when a lobby member leaves a joined lobby.
+	 *
+	 * @return Event that can be bound to
+	 */
+	virtual TOnlineEvent<void(const FLobbyMemberLeft&)> OnLobbyMemberLeft() = 0;
+
+	/**
+	 * Get the event that is triggered when the leadership of a lobby changes.
+	 *
+	 * @return Event that can be bound to
+	 */
+	virtual TOnlineEvent<void(const FLobbyLeaderChanged&)> OnLobbyLeaderChanged() = 0;
+
+	/**
+	 * Get the event that is triggered when the attribute schema of a lobby changes.
+	 *
+	 * @return Event that can be bound to
+	 */
+	virtual TOnlineEvent<void(const FLobbySchemaChanged&)> OnLobbySchemaChanged() = 0;
+
+	/**
+	 * Get the event that is triggered when lobby attributes have changed.
+	 *
+	 * @return Event that can be bound to
+	 */
+	virtual TOnlineEvent<void(const FLobbyAttributesChanged&)> OnLobbyAttributesChanged() = 0;
+
+	/**
+	 * Get the event that is triggered when lobby member attributes have changed.
+	 *
+	 * @return Event that can be bound to
+	 */
+	virtual TOnlineEvent<void(const FLobbyMemberAttributesChanged&)> OnLobbyMemberAttributesChanged() = 0;
+
+	/**
+	 * Get the event that is triggered when an invitation is received.
+	 *
+	 * @return Event that can be bound to
+	 */
+	virtual TOnlineEvent<void(const FLobbyInvitationAdded&)> OnLobbyInvitationAdded() = 0;
+
+	/**
+	 * Get the event that is triggered when an invitation is removed.
+	 *
+	 * @return Event that can be bound to
+	 */
+	virtual TOnlineEvent<void(const FLobbyInvitationRemoved&)> OnLobbyInvitationRemoved() = 0;
+};
+
+namespace Meta {
+// TODO: Move to Lobbies_Meta.inl file?
+
+BEGIN_ONLINE_STRUCT_META(FLobby)
+	ONLINE_STRUCT_FIELD(FLobby, LobbyId),
+	ONLINE_STRUCT_FIELD(FLobby, LocalName),
+	ONLINE_STRUCT_FIELD(FLobby, SchemaName),
+	ONLINE_STRUCT_FIELD(FLobby, MaxMembers),
+	ONLINE_STRUCT_FIELD(FLobby, JoinPolicy)
+END_ONLINE_STRUCT_META()
+
+BEGIN_ONLINE_STRUCT_META(FJoinLobbyLocalUserData)
+	ONLINE_STRUCT_FIELD(FJoinLobbyLocalUserData, LocalUserId),
+	ONLINE_STRUCT_FIELD(FJoinLobbyLocalUserData, Attributes)
+END_ONLINE_STRUCT_META()
+
+BEGIN_ONLINE_STRUCT_META(FCreateLobby::Params)
+	ONLINE_STRUCT_FIELD(FCreateLobby::Params, LocalUserId),
+	ONLINE_STRUCT_FIELD(FCreateLobby::Params, LocalName),
+	ONLINE_STRUCT_FIELD(FCreateLobby::Params, SchemaName),
+	ONLINE_STRUCT_FIELD(FCreateLobby::Params, MaxMembers),
+	ONLINE_STRUCT_FIELD(FCreateLobby::Params, JoinPolicy),
+	ONLINE_STRUCT_FIELD(FCreateLobby::Params, Attributes),
+	ONLINE_STRUCT_FIELD(FCreateLobby::Params, LocalUsers)
+END_ONLINE_STRUCT_META()
+
+BEGIN_ONLINE_STRUCT_META(FCreateLobby::Result)
+	ONLINE_STRUCT_FIELD(FCreateLobby::Result, Lobby)
+END_ONLINE_STRUCT_META()
+
+BEGIN_ONLINE_STRUCT_META(FFindLobby::Params)
+END_ONLINE_STRUCT_META()
+
+BEGIN_ONLINE_STRUCT_META(FFindLobby::Result)
+END_ONLINE_STRUCT_META()
+
+BEGIN_ONLINE_STRUCT_META(FRestoreLobbies::Params)
+END_ONLINE_STRUCT_META()
+
+BEGIN_ONLINE_STRUCT_META(FRestoreLobbies::Result)
+END_ONLINE_STRUCT_META()
+
+BEGIN_ONLINE_STRUCT_META(FJoinLobby::Params)
+END_ONLINE_STRUCT_META()
+
+BEGIN_ONLINE_STRUCT_META(FJoinLobby::Result)
+END_ONLINE_STRUCT_META()
+
+BEGIN_ONLINE_STRUCT_META(FLeaveLobby::Params)
+	ONLINE_STRUCT_FIELD(FLeaveLobby::Params, LocalUserId),
+	ONLINE_STRUCT_FIELD(FLeaveLobby::Params, LobbyId)
+END_ONLINE_STRUCT_META()
+
+BEGIN_ONLINE_STRUCT_META(FLeaveLobby::Result)
+END_ONLINE_STRUCT_META()
+
+BEGIN_ONLINE_STRUCT_META(FInviteLobbyMember::Params)
+END_ONLINE_STRUCT_META()
+
+BEGIN_ONLINE_STRUCT_META(FInviteLobbyMember::Result)
+END_ONLINE_STRUCT_META()
+
+BEGIN_ONLINE_STRUCT_META(FDeclineLobbyInvitation::Params)
+END_ONLINE_STRUCT_META()
+
+BEGIN_ONLINE_STRUCT_META(FDeclineLobbyInvitation::Result)
+END_ONLINE_STRUCT_META()
+
+BEGIN_ONLINE_STRUCT_META(FKickLobbyMember::Params)
+END_ONLINE_STRUCT_META()
+
+BEGIN_ONLINE_STRUCT_META(FKickLobbyMember::Result)
+END_ONLINE_STRUCT_META()
+
+BEGIN_ONLINE_STRUCT_META(FPromoteLobbyMember::Params)
+END_ONLINE_STRUCT_META()
+
+BEGIN_ONLINE_STRUCT_META(FPromoteLobbyMember::Result)
+END_ONLINE_STRUCT_META()
+
+BEGIN_ONLINE_STRUCT_META(FModifyLobbySchema::Params)
+END_ONLINE_STRUCT_META()
+
+BEGIN_ONLINE_STRUCT_META(FModifyLobbySchema::Result)
+END_ONLINE_STRUCT_META()
+
+BEGIN_ONLINE_STRUCT_META(FModifyLobbyJoinPolicy::Params)
+END_ONLINE_STRUCT_META()
+
+BEGIN_ONLINE_STRUCT_META(FModifyLobbyJoinPolicy::Result)
+END_ONLINE_STRUCT_META()
+
+BEGIN_ONLINE_STRUCT_META(FModifyLobbyAttributes::Params)
+END_ONLINE_STRUCT_META()
+
+BEGIN_ONLINE_STRUCT_META(FModifyLobbyAttributes::Result)
+END_ONLINE_STRUCT_META()
+
+BEGIN_ONLINE_STRUCT_META(FModifyLobbyMemberAttributes::Params)
+END_ONLINE_STRUCT_META()
+
+BEGIN_ONLINE_STRUCT_META(FModifyLobbyMemberAttributes::Result)
+END_ONLINE_STRUCT_META()
+
+BEGIN_ONLINE_STRUCT_META(FGetJoinedLobbies::Params)
+END_ONLINE_STRUCT_META()
+
+BEGIN_ONLINE_STRUCT_META(FGetJoinedLobbies::Result)
+END_ONLINE_STRUCT_META()
+
+BEGIN_ONLINE_STRUCT_META(FGetReceivedInvitations::Params)
+END_ONLINE_STRUCT_META()
+
+BEGIN_ONLINE_STRUCT_META(FGetReceivedInvitations::Result)
+END_ONLINE_STRUCT_META()
+
+/* Meta*/ }
+
+/* UE::Online */ }
diff --git a/Engine/Plugins/Online/Experimental/OnlineServicesEOS/Source/Private/LobbiesEOS.cpp b/Engine/Plugins/Online/Experimental/OnlineServicesEOS/Source/Private/LobbiesEOS.cpp
new file mode 100644
--- /dev/null
+++ b/Engine/Plugins/Online/Experimental/OnlineServicesEOS/Source/Private/LobbiesEOS.cpp
@@ -0,0 +1,1511 @@
+// Copyright Epic Games, Inc. All Rights Reserved.
+
+#include "LobbiesEOS.h"
+#include "LobbiesEOSTypes.h"
+#include "Online/OnlineAsyncOp.h"
+#include "OnlineIdEOS.h"
+#include "OnlineServicesEOS.h"
+#include "OnlineServicesEOSTypes.h"
+#include "AuthEOS.h"
+#include "IEOSSDKManager.h"
+
+#include "eos_lobby.h"
+
+DEFINE_LOG_CATEGORY(LogLobbies);
+
+namespace UE::Online {
+
+static const FString LobbyDataKeyName = TEXT("LobbyData");
+static const FString LobbyNotificationPauseKeyName = TEXT("LobbyNotificationPause");
+
+static const int32 MaxAttributeSize = 1000;
+
+template <typename DataType, typename OpType>
+const DataType& GetOpDataChecked(const TOnlineAsyncOp<OpType>& Op, const FString& Key)
+{
+	const DataType* Data = Op.Data.template Get<DataType>(Key);
+	check(Data);
+	return *Data;
+}
+
+FLobbiesEOS::FLobbiesEOS(FOnlineServicesEOS& InServices)
+	: FLobbiesCommon(InServices)
+{
+}
+
+void FLobbiesEOS::Initialize()
+{
+	FLobbiesCommon::Initialize();
+
+	IEOSSDKManager* SDKManager = IEOSSDKManager::Get();
+	check(SDKManager);
+
+	LobbyBucketId = FString::Printf(TEXT("%s_%s"), *SDKManager->GetProductName(), *SDKManager->GetProductVersion());
+
+	EOS_HLobby LobbyInterfaceHandle = EOS_Platform_GetLobbyInterface(static_cast<FOnlineServicesEOS&>(GetServices()).GetEOSPlatformHandle());
+	check(LobbyInterfaceHandle != nullptr);
+
+	LobbyPrerequisites = MakeShared<FLobbyPrerequisitesEOS>(FLobbyPrerequisitesEOS{
+		LobbyInterfaceHandle,
+		StaticCastSharedPtr<FAuthEOS>(Services.GetAuthInterface()).ToSharedRef(),
+		LobbySchemaRegistry.ToSharedRef(),
+		ServiceSchema.ToSharedRef()});
+
+	LobbyDataRegistry = MakeShared<FLobbyDataRegistryEOS>(LobbyPrerequisites.ToSharedRef());
+
+	RegisterHandlers();
+}
+
+void FLobbiesEOS::PreShutdown()
+{
+	UnregisterHandlers();
+
+	LobbyDataRegistry = nullptr;
+	LobbyPrerequisites = nullptr;
+}
+
+TOnlineAsyncOpHandle<FCreateLobby> FLobbiesEOS::CreateLobby(FCreateLobby::Params&& InParams)
+{
+	TOnlineAsyncOpRef<FCreateLobby> Op = GetOp<FCreateLobby>(MoveTemp(InParams));
+	const FCreateLobby::Params& Params = Op->GetParams();
+
+	if (!Services.Get<FAuthEOS>()->IsLoggedIn(Params.LocalUserId))
+	{
+		Op->SetError(Errors::InvalidUser());
+		return Op->GetHandle();
+	}
+
+	auto DestroyLobbyDuringCreate = [this](
+		TOnlineAsyncOp<FCreateLobby>& InAsyncOp,
+		const TSharedPtr<FLobbyDataEOS>& LobbyData,
+		FOnlineAccountIdHandle LocalUserId,
+		FOnlineError ErrorResult) -> TFuture<void>
+	{
+		FDestroyLobbyImpl::Params DestroyLobbyParams;
+		DestroyLobbyParams.LobbyData = LobbyData;
+		DestroyLobbyParams.LocalUserId = LocalUserId;
+
+		TPromise<void> Promise;
+		auto Future = Promise.GetFuture();
+
+		DestroyLobbyImpl(MoveTemp(DestroyLobbyParams))
+		.Then([this, InAsyncOp = InAsyncOp.AsShared(), ErrorResult = MoveTemp(ErrorResult), Promise = MoveTemp(Promise)](TFuture<TDefaultErrorResult<FDestroyLobbyImpl>>&& Future) mutable
+		{
+			if (Future.Get().IsError())
+			{
+				// Todo: complain about having an error while handling an error.
+			}
+
+			// Todo: Errors.
+			InAsyncOp->SetError(Errors::Unknown(MoveTemp(ErrorResult)));
+			Promise.EmplaceValue();
+		});
+
+		return Future;
+	};
+
+	Op->Then([this](TOnlineAsyncOp<FCreateLobby>& InAsyncOp)
+	{
+		const auto& Params = InAsyncOp.GetParams();
+		const FTCHARToUTF8 BucketIdUTF8(*LobbyBucketId);
+
+		// The lobby will be created as invitation only.
+		// Once all local members are joined and the lobby attributes have been set the privacy setting will be moved to the user setting.
+
+		EOS_Lobby_CreateLobbyOptions CreateLobbyOptions = {};
+		CreateLobbyOptions.ApiVersion = EOS_LOBBY_CREATELOBBY_API_LATEST;
+		CreateLobbyOptions.LocalUserId = GetProductUserIdChecked(Params.LocalUserId);
+		CreateLobbyOptions.MaxLobbyMembers = Params.MaxMembers;
+		CreateLobbyOptions.PermissionLevel = TranslateJoinPolicy(ELobbyJoinPolicy::InvitationOnly);
+		CreateLobbyOptions.bPresenceEnabled = false; // todo: handle
+		CreateLobbyOptions.bAllowInvites = true; // todo: handle
+		CreateLobbyOptions.BucketId = BucketIdUTF8.Get();
+		CreateLobbyOptions.bDisableHostMigration = false; // todo: handle
+		CreateLobbyOptions.bEnableRTCRoom = false; // todo: handle
+
+		UE_LOG(LogLobbies, Warning, TEXT("[FLobbiesEOS::EOS_Lobby_CreateLobby] Start. Member: %s"), *ToLogString(Params.LocalUserId));
+
+		return EOS_Async<EOS_Lobby_CreateLobbyCallbackInfo>(EOS_Lobby_CreateLobby, LobbyPrerequisites->LobbyInterfaceHandle, CreateLobbyOptions); 
+	})
+	.Then([this](TOnlineAsyncOp<FCreateLobby>& InAsyncOp, const EOS_Lobby_CreateLobbyCallbackInfo* Data)
+	{
+		UE_LOG(LogLobbies, Warning, TEXT("[FLobbiesEOS::EOS_Lobby_CreateLobby] Complete. Result: %d"), Data->ResultCode);
+
+		if (Data->ResultCode != EOS_EResult::EOS_Success)
+		{
+			// Todo: errors
+			InAsyncOp.SetError(Errors::Unknown(FromEOSError(Data->ResultCode)));
+			return MakeFulfilledPromise<TDefaultErrorResultInternal<TSharedPtr<FLobbyDataEOS>>>().GetFuture();
+		}
+		
+		return LobbyDataRegistry->FindOrCreateFromLobbyId(InAsyncOp.GetParams().LocalUserId, Data->LobbyId);
+	})
+	.Then([this, DestroyLobbyDuringCreate](TOnlineAsyncOp<FCreateLobby>& InAsyncOp, TDefaultErrorResultInternal<TSharedPtr<FLobbyDataEOS>>&& Result)
+	{
+		const auto& Params = InAsyncOp.GetParams();
+
+		if (Result.IsError())
+		{
+			// Todo: destroy lobby.
+
+			// Todo: errors
+			InAsyncOp.SetError(Errors::Unknown(MoveTemp(Result.GetErrorValue())));
+			return;
+		}
+
+		TSharedRef<FLobbyDataEOS> LobbyData = Result.GetOkValue().ToSharedRef();
+
+		// Pause notifications for all joining local users.
+		TArray<TSharedRef<FLobbyNotificationPauseHandle>> PauseHandles;
+		for (const FJoinLobbyLocalUserData& UserData : Params.LocalUsers)
+		{
+			PauseHandles.Add(LobbyData->PauseLobbyNotifications(UserData.LocalUserId));
+		}
+
+		InAsyncOp.Data.Set<TSharedRef<FLobbyDataEOS>>(LobbyDataKeyName, MoveTemp(LobbyData));
+		InAsyncOp.Data.Set<TArray<TSharedRef<FLobbyNotificationPauseHandle>>>(LobbyNotificationPauseKeyName, MoveTemp(PauseHandles));
+	})
+	.Then([this](TOnlineAsyncOp<FCreateLobby>& InAsyncOp)
+	{
+		const auto& Params = InAsyncOp.GetParams();
+		auto LobbyData = GetOpDataChecked<TSharedRef<FLobbyDataEOS>>(InAsyncOp, LobbyDataKeyName);
+
+		// Add creator attributes.
+		FModifyLobbyMemberDataImpl::Params ModifyLobbyMemberDataParams;
+		ModifyLobbyMemberDataParams.LobbyData = LobbyData;
+		ModifyLobbyMemberDataParams.LocalUserId = Params.LocalUserId;
+		ModifyLobbyMemberDataParams.Changes = MakeShared<FLobbyMemberDataChanges>();
+		//ModifyLobbyMemberDataParams.Changes->MutatedAttributes;
+
+		return ModifyLobbyMemberDataImpl(MoveTemp(ModifyLobbyMemberDataParams));
+	})
+	.Then([this, DestroyLobbyDuringCreate](TOnlineAsyncOp<FCreateLobby>& InAsyncOp, TDefaultErrorResult<FModifyLobbyMemberDataImpl>&& ModifyLobbyOwnerResult)
+	{
+		const auto& Params = InAsyncOp.GetParams();
+		auto LobbyData = GetOpDataChecked<TSharedRef<FLobbyDataEOS>>(InAsyncOp, LobbyDataKeyName);
+
+		if (ModifyLobbyOwnerResult.IsError())
+		{
+			return DestroyLobbyDuringCreate(InAsyncOp, LobbyData, Params.LocalUserId, MoveTemp(ModifyLobbyOwnerResult.GetErrorValue()));
+		}
+		else
+		{
+			return MakeFulfilledPromise<void>().GetFuture();
+		}
+	})
+	.Then([this](TOnlineAsyncOp<FCreateLobby>& InAsyncOp)
+	{
+		// Todo: Add other local members.
+	})
+	.Then([this](TOnlineAsyncOp<FCreateLobby>& InAsyncOp)
+	{
+		const auto& Params = InAsyncOp.GetParams();
+		auto LobbyData = GetOpDataChecked<TSharedRef<FLobbyDataEOS>>(InAsyncOp, LobbyDataKeyName);
+
+		// Add lobby attributes.
+		// Set lobby privacy setting to the user provided value.
+
+		FModifyLobbyDataImpl::Params ModifyLobbyDataParams;
+		ModifyLobbyDataParams.LobbyData = LobbyData;
+		ModifyLobbyDataParams.LocalUserId = Params.LocalUserId;
+		ModifyLobbyDataParams.Changes = MakeShared<FLobbyDataChanges>();
+		ModifyLobbyDataParams.Changes->MutatedAttributes = Params.Attributes;
+		ModifyLobbyDataParams.Changes->JoinPolicy = Params.JoinPolicy;
+
+		return ModifyLobbyDataImpl(MoveTemp(ModifyLobbyDataParams));
+	})
+	.Then([this, DestroyLobbyDuringCreate](TOnlineAsyncOp<FCreateLobby>& InAsyncOp, TDefaultErrorResult<FModifyLobbyDataImpl>&& ModifyLobbyResult) mutable
+	{
+		const auto& Params = InAsyncOp.GetParams();
+		auto LobbyData = GetOpDataChecked<TSharedRef<FLobbyDataEOS>>(InAsyncOp, LobbyDataKeyName);
+
+		if (ModifyLobbyResult.IsError())
+		{
+			return DestroyLobbyDuringCreate(InAsyncOp, LobbyData, Params.LocalUserId, MoveTemp(ModifyLobbyResult.GetErrorValue()));
+		}
+		else
+		{
+			return MakeFulfilledPromise<void>().GetFuture();
+		}
+	})
+	.Then([this](TOnlineAsyncOp<FCreateLobby>& InAsyncOp)
+	{
+		const auto& Params = InAsyncOp.GetParams();
+		auto LobbyData = GetOpDataChecked<TSharedRef<FLobbyDataEOS>>(InAsyncOp, LobbyDataKeyName);
+
+		// Remove notification pause for joining local users.
+		TArray<TSharedRef<FLobbyNotificationPauseHandle>> ClearedPauses;
+		InAsyncOp.Data.Set<TArray<TSharedRef<FLobbyNotificationPauseHandle>>>(LobbyNotificationPauseKeyName, MoveTemp(ClearedPauses));
+
+		// Mark the lobby active for each member.
+		// Lobby activation / deactivation will dispatch lobby join / leave notifications.
+		for (const FJoinLobbyLocalUserData& LocalUser : Params.LocalUsers)
+		{
+			AddActiveLobby(LocalUser.LocalUserId, LobbyData);
+		}
+
+		InAsyncOp.SetResult(FCreateLobby::Result{LobbyData->GetLobbyImpl()});
+	})
+	.Enqueue(GetSerialQueue(Params.LocalUserId));
+
+	return Op->GetHandle();
+}
+
+TOnlineAsyncOpHandle<FFindLobby> FLobbiesEOS::FindLobby(FFindLobby::Params&& InParams)
+{
+	TOnlineAsyncOpRef<FFindLobby> Op = GetOp<FFindLobby>(MoveTemp(InParams));
+	const FFindLobby::Params& Params = Op->GetParams();
+
+	if (!Services.Get<FAuthEOS>()->IsLoggedIn(Params.LocalUserId))
+	{
+		Op->SetError(Errors::InvalidUser());
+		return Op->GetHandle();
+	}
+
+	// Invalidate previous search results.
+	ActiveSearchResults.Remove(Params.LocalUserId);
+
+	Op->Then([this](TOnlineAsyncOp<FFindLobby>& InAsyncOp)
+	{
+		return FLobbySearchEOS::Create(LobbyPrerequisites.ToSharedRef(), LobbyDataRegistry.ToSharedRef(), InAsyncOp.GetParams());
+	})
+	.Then([this](TOnlineAsyncOp<FFindLobby>& InAsyncOp, TDefaultErrorResultInternal<TSharedPtr<FLobbySearchEOS>> Result)
+	{
+		if (Result.IsError())
+		{
+			// TODO: Error codes
+			InAsyncOp.SetError(Errors::Unknown(Result.GetErrorValue()));
+		}
+		else
+		{
+			ActiveSearchResults.Add(InAsyncOp.GetParams().LocalUserId, Result.GetOkValue().ToSharedRef());
+			InAsyncOp.SetResult(FFindLobby::Result{Result.GetOkValue()->GetLobbyResults()});
+		}
+	})
+	.Enqueue(GetSerialQueue(Params.LocalUserId));
+
+	return Op->GetHandle();
+}
+
+TOnlineAsyncOpHandle<FJoinLobby> FLobbiesEOS::JoinLobby(FJoinLobby::Params&& InParams)
+{
+	TOnlineAsyncOpRef<FJoinLobby> Op = GetOp<FJoinLobby>(MoveTemp(InParams));
+	const FJoinLobby::Params& Params = Op->GetParams();
+
+	if (!Services.Get<FAuthEOS>()->IsLoggedIn(Params.LocalUserId))
+	{
+		Op->SetError(Errors::InvalidUser());
+		return Op->GetHandle();
+	}
+
+	TSharedPtr<FLobbyDataEOS> LobbyData = LobbyDataRegistry->Find(Params.LobbyId);
+	if (!LobbyData)
+	{
+		Op->SetError(Errors::Unknown());
+		return Op->GetHandle();
+	}
+
+	Op->Then([this, LobbyData](TOnlineAsyncOp<FJoinLobby>& InAsyncOp)
+	{
+		const FJoinLobby::Params& Params = InAsyncOp.GetParams();
+
+		// Pause notifications for all joining local users.
+		TArray<TSharedRef<FLobbyNotificationPauseHandle>> PauseHandles;
+		for (const FJoinLobbyLocalUserData& UserData : Params.LocalUsers)
+		{
+			PauseHandles.Add(LobbyData->PauseLobbyNotifications(UserData.LocalUserId));
+		}
+		InAsyncOp.Data.Set<TArray<TSharedRef<FLobbyNotificationPauseHandle>>>(LobbyNotificationPauseKeyName, MoveTemp(PauseHandles));
+
+		FJoinLobbyImpl::Params JoinParams;
+		JoinParams.LobbyData = LobbyData;
+		JoinParams.LocalUserId = Params.LocalUserId;
+		JoinParams.LocalName = Params.LocalName;
+		JoinParams.LocalUsers = Params.LocalUsers;
+		return JoinLobbyImpl(MoveTemp(JoinParams));
+	})
+	.Then([this, LobbyData](TOnlineAsyncOp<FJoinLobby>& InAsyncOp, TDefaultErrorResult<FJoinLobbyImpl> Result)
+	{
+		const auto& Params = InAsyncOp.GetParams();
+
+		if (Result.IsError())
+		{
+			// TODO: Error codes
+			InAsyncOp.SetError(Errors::Unknown(Result.GetErrorValue()));
+		}
+		else
+		{
+			// Remove notification pause for joining local users.
+			TArray<TSharedRef<FLobbyNotificationPauseHandle>> ClearedPauses;
+			InAsyncOp.Data.Set<TArray<TSharedRef<FLobbyNotificationPauseHandle>>>(LobbyNotificationPauseKeyName, MoveTemp(ClearedPauses));
+
+			// Mark the lobby active for each member.
+			// Lobby activation / deactivation will dispatch lobby join / leave notifications.
+			for (const FJoinLobbyLocalUserData& LocalUser : Params.LocalUsers)
+			{
+				AddActiveLobby(LocalUser.LocalUserId, LobbyData.ToSharedRef());
+			}
+
+			InAsyncOp.SetResult(FJoinLobby::Result{});
+		}
+	})
+	.Enqueue(GetSerialQueue());
+
+	return Op->GetHandle();
+}
+
+TOnlineAsyncOpHandle<FLeaveLobby> FLobbiesEOS::LeaveLobby(FLeaveLobby::Params&& InParams)
+{
+	TOnlineAsyncOpRef<FLeaveLobby> Op = GetOp<FLeaveLobby>(MoveTemp(InParams));
+	const FLeaveLobby::Params& Params = Op->GetParams();
+
+	if (!Services.Get<FAuthEOS>()->IsLoggedIn(Params.LocalUserId))
+	{
+		Op->SetError(Errors::InvalidUser());
+		return Op->GetHandle();
+	}
+
+	TSharedPtr<FLobbyDataEOS> LobbyData = LobbyDataRegistry->Find(Params.LobbyId);
+	if (!LobbyData)
+	{
+		Op->SetError(Errors::Unknown());
+		return Op->GetHandle();
+	}
+
+	Op->Then([this, LobbyData](TOnlineAsyncOp<FLeaveLobby>& InAsyncOp)
+	{
+		const FLeaveLobby::Params& Params = InAsyncOp.GetParams();
+
+		FLeaveLobbyImpl::Params LeaveParams;
+		LeaveParams.LobbyData = LobbyData;
+		LeaveParams.LocalUserId = Params.LocalUserId;
+		return LeaveLobbyImpl(MoveTemp(LeaveParams));
+	})
+	.Then([this, LobbyData](TOnlineAsyncOp<FLeaveLobby>& InAsyncOp, TDefaultErrorResult<FLeaveLobbyImpl> Result)
+	{
+		const auto& Params = InAsyncOp.GetParams();
+
+		if (LobbyData)
+		{
+			RemoveActiveLobby(Params.LocalUserId, LobbyData.ToSharedRef());
+		}
+
+		if (Result.IsError())
+		{
+			// TODO: Error codes
+			InAsyncOp.SetError(Errors::Unknown(Result.GetErrorValue()));
+		}
+		else
+		{
+			InAsyncOp.SetResult(FLeaveLobby::Result{});
+		}
+	})
+	.Enqueue(GetSerialQueue());
+
+	return Op->GetHandle();
+}
+
+TOnlineAsyncOpHandle<FInviteLobbyMember> FLobbiesEOS::InviteLobbyMember(FInviteLobbyMember::Params&& InParams)
+{
+	TOnlineAsyncOpRef<FInviteLobbyMember> Op = GetOp<FInviteLobbyMember>(MoveTemp(InParams));
+	const FInviteLobbyMember::Params& Params = Op->GetParams();
+
+	if (!Services.Get<FAuthEOS>()->IsLoggedIn(Params.LocalUserId))
+	{
+		Op->SetError(Errors::InvalidUser());
+		return Op->GetHandle();
+	}
+
+	Op->Then([this](TOnlineAsyncOp<FInviteLobbyMember>& InAsyncOp)
+	{
+		const FInviteLobbyMember::Params& Params = InAsyncOp.GetParams();
+
+		FInviteLobbyMemberImpl::Params InviteParams;
+		InviteParams.LobbyData = LobbyDataRegistry->Find(Params.LobbyId);
+		InviteParams.LocalUserId = Params.LocalUserId;
+		InviteParams.TargetUserId = Params.TargetUserId;
+		return InviteLobbyMemberImpl(MoveTemp(InviteParams));
+	})
+	.Then([this](TOnlineAsyncOp<FInviteLobbyMember>& InAsyncOp, TDefaultErrorResult<FInviteLobbyMemberImpl> Result)
+	{
+		if (Result.IsError())
+		{
+			// TODO: Error codes
+			InAsyncOp.SetError(Errors::Unknown(Result.GetErrorValue()));
+		}
+		else
+		{
+			InAsyncOp.SetResult(FInviteLobbyMember::Result{});
+		}
+	})
+	.Enqueue(GetSerialQueue(Params.LocalUserId));
+
+	return Op->GetHandle();
+}
+
+TOnlineAsyncOpHandle<FDeclineLobbyInvitation> FLobbiesEOS::DeclineLobbyInvitation(FDeclineLobbyInvitation::Params&& InParams)
+{
+	TOnlineAsyncOpRef<FDeclineLobbyInvitation> Op = GetOp<FDeclineLobbyInvitation>(MoveTemp(InParams));
+	const FDeclineLobbyInvitation::Params& Params = Op->GetParams();
+
+	if (!Services.Get<FAuthEOS>()->IsLoggedIn(Params.LocalUserId))
+	{
+		Op->SetError(Errors::InvalidUser());
+		return Op->GetHandle();
+	}
+
+	Op->Then([this](TOnlineAsyncOp<FDeclineLobbyInvitation>& InAsyncOp)
+	{
+		const FDeclineLobbyInvitation::Params& Params = InAsyncOp.GetParams();
+
+		FDeclineLobbyInvitationImpl::Params DeclineParams;
+		DeclineParams.LocalUserId = Params.LocalUserId;
+		DeclineParams.LobbyId = Params.LobbyId;
+		return DeclineLobbyInvitationImpl(MoveTemp(DeclineParams));
+	})
+	.Then([this](TOnlineAsyncOp<FDeclineLobbyInvitation>& InAsyncOp, TDefaultErrorResult<FDeclineLobbyInvitationImpl> Result)
+	{
+		if (Result.IsError())
+		{
+			// TODO: Error codes
+			InAsyncOp.SetError(Errors::Unknown(Result.GetErrorValue()));
+		}
+		else
+		{
+			InAsyncOp.SetResult(FDeclineLobbyInvitation::Result{});
+		}
+	})
+	.Enqueue(GetSerialQueue(Params.LocalUserId));
+
+	return Op->GetHandle();
+}
+
+TOnlineAsyncOpHandle<FKickLobbyMember> FLobbiesEOS::KickLobbyMember(FKickLobbyMember::Params&& InParams)
+{
+	TOnlineAsyncOpRef<FKickLobbyMember> Op = GetOp<FKickLobbyMember>(MoveTemp(InParams));
+	const FKickLobbyMember::Params& Params = Op->GetParams();
+
+	if (!Services.Get<FAuthEOS>()->IsLoggedIn(Params.LocalUserId))
+	{
+		Op->SetError(Errors::InvalidUser());
+		return Op->GetHandle();
+	}
+
+	Op->Then([this](TOnlineAsyncOp<FKickLobbyMember>& InAsyncOp)
+	{
+		const FKickLobbyMember::Params& Params = InAsyncOp.GetParams();
+
+		FKickLobbyMemberImpl::Params KickParams;
+		KickParams.LobbyData = LobbyDataRegistry->Find(Params.LobbyId);
+		KickParams.LocalUserId = Params.LocalUserId;
+		KickParams.TargetUserId = Params.TargetUserId;
+		return KickLobbyMemberImpl(MoveTemp(KickParams));
+	})
+	.Then([this](TOnlineAsyncOp<FKickLobbyMember>& InAsyncOp, TDefaultErrorResult<FKickLobbyMemberImpl> Result)
+	{
+		if (Result.IsError())
+		{
+			// TODO: Error codes
+			InAsyncOp.SetError(Errors::Unknown(Result.GetErrorValue()));
+		}
+		else
+		{
+			InAsyncOp.SetResult(FKickLobbyMember::Result{});
+		}
+	})
+	.Enqueue(GetSerialQueue(Params.LocalUserId));
+
+	return Op->GetHandle();
+}
+
+TOnlineAsyncOpHandle<FPromoteLobbyMember> FLobbiesEOS::PromoteLobbyMember(FPromoteLobbyMember::Params&& InParams)
+{
+	TOnlineAsyncOpRef<FPromoteLobbyMember> Op = GetOp<FPromoteLobbyMember>(MoveTemp(InParams));
+	const FPromoteLobbyMember::Params& Params = Op->GetParams();
+
+	if (!Services.Get<FAuthEOS>()->IsLoggedIn(Params.LocalUserId))
+	{
+		Op->SetError(Errors::InvalidUser());
+		return Op->GetHandle();
+	}
+
+	Op->Then([this](TOnlineAsyncOp<FPromoteLobbyMember>& InAsyncOp)
+	{
+		const FPromoteLobbyMember::Params& Params = InAsyncOp.GetParams();
+
+		FPromoteLobbyMemberImpl::Params PromoteParams;
+		PromoteParams.LobbyData = LobbyDataRegistry->Find(Params.LobbyId);
+		PromoteParams.LocalUserId = Params.LocalUserId;
+		PromoteParams.TargetUserId = Params.TargetUserId;
+		return PromoteLobbyMemberImpl(MoveTemp(PromoteParams));
+	})
+	.Then([this](TOnlineAsyncOp<FPromoteLobbyMember>& InAsyncOp, TDefaultErrorResult<FPromoteLobbyMemberImpl> Result)
+	{
+		if (Result.IsError())
+		{
+			// TODO: Error codes
+			InAsyncOp.SetError(Errors::Unknown(Result.GetErrorValue()));
+		}
+		else
+		{
+			InAsyncOp.SetResult(FPromoteLobbyMember::Result{});
+		}
+	})
+	.Enqueue(GetSerialQueue(Params.LocalUserId));
+
+	return Op->GetHandle();
+}
+
+TOnlineAsyncOpHandle<FModifyLobbyJoinPolicy> FLobbiesEOS::ModifyLobbyJoinPolicy(FModifyLobbyJoinPolicy::Params&& InParams)
+{
+	TOnlineAsyncOpRef<FModifyLobbyJoinPolicy> Op = GetOp<FModifyLobbyJoinPolicy>(MoveTemp(InParams));
+	const FModifyLobbyJoinPolicy::Params& Params = Op->GetParams();
+
+	if (!Services.Get<FAuthEOS>()->IsLoggedIn(Params.LocalUserId))
+	{
+		Op->SetError(Errors::InvalidUser());
+		return Op->GetHandle();
+	}
+
+	Op->Then([this](TOnlineAsyncOp<FModifyLobbyJoinPolicy>& InAsyncOp)
+	{
+		const FModifyLobbyJoinPolicy::Params& Params = InAsyncOp.GetParams();
+
+		FModifyLobbyDataImpl::Params ModifyLobbyDataParams;
+		ModifyLobbyDataParams.LobbyData = LobbyDataRegistry->Find(Params.LobbyId);
+		ModifyLobbyDataParams.LocalUserId = Params.LocalUserId;
+		ModifyLobbyDataParams.Changes = MakeShared<FLobbyDataChanges>();
+		ModifyLobbyDataParams.Changes->JoinPolicy = Params.JoinPolicy;
+		return ModifyLobbyDataImpl(MoveTemp(ModifyLobbyDataParams));
+	})
+	.Then([this](TOnlineAsyncOp<FModifyLobbyJoinPolicy>& InAsyncOp, TDefaultErrorResult<FModifyLobbyDataImpl> Result)
+	{
+		if (Result.IsError())
+		{
+			// TODO: Error codes
+			InAsyncOp.SetError(Errors::Unknown(Result.GetErrorValue()));
+		}
+		else
+		{
+			InAsyncOp.SetResult(FModifyLobbyJoinPolicy::Result{});
+		}
+	})
+	.Enqueue(GetSerialQueue(Params.LocalUserId));
+
+	return Op->GetHandle();
+}
+
+TOnlineAsyncOpHandle<FModifyLobbyAttributes> FLobbiesEOS::ModifyLobbyAttributes(FModifyLobbyAttributes::Params&& InParams)
+{
+	TOnlineAsyncOpRef<FModifyLobbyAttributes> Op = GetOp<FModifyLobbyAttributes>(MoveTemp(InParams));
+	const FModifyLobbyAttributes::Params& Params = Op->GetParams();
+
+	if (!Services.Get<FAuthEOS>()->IsLoggedIn(Params.LocalUserId))
+	{
+		Op->SetError(Errors::InvalidUser());
+		return Op->GetHandle();
+	}
+
+	Op->Then([this](TOnlineAsyncOp<FModifyLobbyAttributes>& InAsyncOp)
+	{
+		const FModifyLobbyAttributes::Params& Params = InAsyncOp.GetParams();
+
+		FModifyLobbyDataImpl::Params ModifyLobbyDataParams;
+		ModifyLobbyDataParams.LobbyData = LobbyDataRegistry->Find(Params.LobbyId);
+		ModifyLobbyDataParams.LocalUserId = Params.LocalUserId;
+		ModifyLobbyDataParams.Changes = MakeShared<FLobbyDataChanges>();
+		ModifyLobbyDataParams.Changes->MutatedAttributes = Params.MutatedAttributes;
+		ModifyLobbyDataParams.Changes->ClearedAttributes = Params.ClearedAttributes;
+		return ModifyLobbyDataImpl(MoveTemp(ModifyLobbyDataParams));
+	})
+	.Then([this](TOnlineAsyncOp<FModifyLobbyAttributes>& InAsyncOp, TDefaultErrorResult<FModifyLobbyDataImpl> Result)
+	{
+		if (Result.IsError())
+		{
+			// TODO: Error codes
+			InAsyncOp.SetError(Errors::Unknown(Result.GetErrorValue()));
+		}
+		else
+		{
+			InAsyncOp.SetResult(FModifyLobbyAttributes::Result{});
+		}
+	})
+	.Enqueue(GetSerialQueue(Params.LocalUserId));
+
+	return Op->GetHandle();
+}
+
+TOnlineAsyncOpHandle<FModifyLobbyMemberAttributes> FLobbiesEOS::ModifyLobbyMemberAttributes(FModifyLobbyMemberAttributes::Params&& InParams)
+{
+	TOnlineAsyncOpRef<FModifyLobbyMemberAttributes> Op = GetOp<FModifyLobbyMemberAttributes>(MoveTemp(InParams));
+	const FModifyLobbyMemberAttributes::Params& Params = Op->GetParams();
+
+	if (!Services.Get<FAuthEOS>()->IsLoggedIn(Params.LocalUserId))
+	{
+		Op->SetError(Errors::InvalidUser());
+		return Op->GetHandle();
+	}
+
+	Op->Then([this](TOnlineAsyncOp<FModifyLobbyMemberAttributes>& InAsyncOp)
+	{
+		const FModifyLobbyMemberAttributes::Params& Params = InAsyncOp.GetParams();
+
+		FModifyLobbyMemberDataImpl::Params ModifyLobbyMemberDataParams;
+		ModifyLobbyMemberDataParams.LobbyData = LobbyDataRegistry->Find(Params.LobbyId);
+		ModifyLobbyMemberDataParams.LocalUserId = Params.LocalUserId;
+		ModifyLobbyMemberDataParams.Changes = MakeShared<FLobbyMemberDataChanges>();
+		ModifyLobbyMemberDataParams.Changes->MutatedAttributes = Params.MutatedAttributes;
+		ModifyLobbyMemberDataParams.Changes->ClearedAttributes = Params.ClearedAttributes;
+		return ModifyLobbyMemberDataImpl(MoveTemp(ModifyLobbyMemberDataParams));
+	})
+	.Then([this](TOnlineAsyncOp<FModifyLobbyMemberAttributes>& InAsyncOp, TDefaultErrorResult<FModifyLobbyMemberDataImpl> Result)
+	{
+		if (Result.IsError())
+		{
+			// TODO: Error codes
+			InAsyncOp.SetError(Errors::Unknown(Result.GetErrorValue()));
+		}
+		else
+		{
+			InAsyncOp.SetResult(FModifyLobbyMemberAttributes::Result{});
+		}
+	})
+	.Enqueue(GetSerialQueue(Params.LocalUserId));
+
+	return Op->GetHandle();
+}
+
+void FLobbiesEOS::HandleLobbyUpdated(const EOS_Lobby_LobbyUpdateReceivedCallbackInfo* Data)
+{
+#if 0
+	EOS_STRUCT(EOS_Lobby_LobbyUpdateReceivedCallbackInfo, (
+		/** Context that was passed into EOS_Lobby_AddNotifyLobbyUpdateReceived */
+		void* ClientData;
+		/** The ID of the lobby */
+		EOS_LobbyId LobbyId;
+	));
+#endif
+
+	if (TSharedPtr<FLobbyDataEOS> LobbyData = LobbyDataRegistry->Find(Data->LobbyId))
+	{
+		UE_LOG(LogLobbies, Warning, TEXT("[FLobbiesEOS::HandleLobbyUpdated] Received update. Lobby: %s"), *LobbyData->GetLobbyId());
+	}
+}
+
+void FLobbiesEOS::HandleLobbyMemberUpdated(const EOS_Lobby_LobbyMemberUpdateReceivedCallbackInfo* Data)
+{
+#if 0
+	EOS_STRUCT(EOS_Lobby_LobbyMemberUpdateReceivedCallbackInfo, (
+		/** Context that was passed into EOS_Lobby_AddNotifyLobbyMemberUpdateReceived */
+		void* ClientData;
+		/** The ID of the lobby */
+		EOS_LobbyId LobbyId;
+		/** The Product User ID of the lobby member */
+		EOS_ProductUserId TargetUserId;
+	));
+#endif
+
+	const FOnlineAccountIdHandle LocalUserId = FindAccountId(Data->TargetUserId);
+	if (LocalUserId.IsValid())
+	{
+		if (TSharedPtr<FLobbyDataEOS> LobbyData = LobbyDataRegistry->Find(Data->LobbyId))
+		{
+			UE_LOG(LogLobbies, Warning, TEXT("[FLobbiesEOS::HandleLobbyMemberUpdated] Received update. Lobby: %s, Member: %s"),
+				*LobbyData->GetLobbyId(), *ToLogString(LocalUserId));
+		}
+	}
+}
+
+void FLobbiesEOS::HandleLobbyMemberStatusReceived(const EOS_Lobby_LobbyMemberStatusReceivedCallbackInfo* Data)
+{
+#if 0
+	EOS_STRUCT(EOS_Lobby_LobbyMemberStatusReceivedCallbackInfo, (
+		/** Context that was passed into EOS_Lobby_AddNotifyLobbyMemberStatusReceived */
+		void* ClientData;
+		/** The ID of the lobby */
+		EOS_LobbyId LobbyId;
+		/** The Product User ID of the lobby member */
+		EOS_ProductUserId TargetUserId;
+		/** Latest status of the user */
+		EOS_ELobbyMemberStatus CurrentStatus;
+	));
+#endif
+
+	const FOnlineAccountIdHandle LocalUserId = FindAccountId(Data->TargetUserId);
+	if (LocalUserId.IsValid())
+	{
+		if (TSharedPtr<FLobbyDataEOS> LobbyData = LobbyDataRegistry->Find(Data->LobbyId))
+		{
+			UE_LOG(LogLobbies, Warning, TEXT("[FLobbiesEOS::HandleLobbyMemberStatusReceived] Received update. Lobby: %s, Member: %s, Statsu: %d"),
+				*LobbyData->GetLobbyId(), *ToLogString(LocalUserId), Data->CurrentStatus);
+		}
+	}
+}
+
+void FLobbiesEOS::HandleLobbyInviteReceived(const EOS_Lobby_LobbyInviteReceivedCallbackInfo* Data)
+{
+#if 0
+	EOS_STRUCT(EOS_Lobby_LobbyInviteReceivedCallbackInfo, (
+		/** Context that was passed into EOS_Lobby_AddNotifyLobbyInviteReceived */
+		void* ClientData;
+		/** The ID of the invitation */
+		const char* InviteId;
+		/** The Product User ID of the local user who received the invitation */
+		EOS_ProductUserId LocalUserId;
+		/** The Product User ID of the user who sent the invitation */
+		EOS_ProductUserId TargetUserId;
+	));
+#endif
+
+	// Todo: Queue this like an operation.
+	const FOnlineAccountIdHandle LocalUserId = FindAccountId(Data->LocalUserId);
+	if (LocalUserId.IsValid())
+	{
+		FLobbyInviteDataEOS::CreateFromInviteId(LobbyPrerequisites.ToSharedRef(), LobbyDataRegistry.ToSharedRef(), LocalUserId, Data->InviteId, Data->TargetUserId)
+		.Then([this](TFuture<TDefaultErrorResultInternal<TSharedPtr<FLobbyInviteDataEOS>>>&& Future)
+		{
+			if (Future.Get().IsError())
+			{
+				// Todo: Log / queue a manual fetch of invitations.
+				UE_LOG(LogLobbies, Warning, TEXT("[FLobbiesEOS::HandleLobbyInviteReceived] Failed to receive invite. Error: %s"),
+					*Future.Get().GetErrorValue().GetLogString());
+			}
+			else
+			{
+				UE_LOG(LogLobbies, Warning, TEXT("[FLobbiesEOS::HandleLobbyInviteReceived] Received invite. Id: %s, Lobby: %s, Receiver: %s, Sender: %s"),
+					*Future.Get().GetOkValue()->GetInviteId(),
+					*Future.Get().GetOkValue()->GetLobbyData()->GetLobbyId(),
+					*ToLogString(Future.Get().GetOkValue()->GetReceiver()),
+					*ToLogString(Future.Get().GetOkValue()->GetSender()));
+				AddActiveInvite(Future.Get().GetOkValue().ToSharedRef());
+			}
+		});
+	}
+}
+
+void FLobbiesEOS::HandleLobbyInviteAccepted(const EOS_Lobby_LobbyInviteAcceptedCallbackInfo* Data)
+{
+#if 0
+	EOS_STRUCT(EOS_Lobby_LobbyInviteAcceptedCallbackInfo, (
+		/** Context that was passed into EOS_Lobby_AddNotifyLobbyInviteAccepted */
+		void* ClientData;
+		/** The invite ID */
+		const char* InviteId;
+		/** The Product User ID of the local user who received the invitation */
+		EOS_ProductUserId LocalUserId;
+		/** The Product User ID of the user who sent the invitation */
+		EOS_ProductUserId TargetUserId;
+		/** Lobby ID that the user has been invited to */
+		const char* LobbyId;
+	));
+#endif
+
+	// Todo: handle catalog of sent invitations.
+}
+
+void FLobbiesEOS::HandleJoinLobbyAccepted(const EOS_Lobby_JoinLobbyAcceptedCallbackInfo* Data)
+{
+#if 0
+	EOS_STRUCT(EOS_Lobby_JoinLobbyAcceptedCallbackInfo, (
+		/** Context that was passed into EOS_Lobby_AddNotifyJoinLobbyAccepted */
+		void* ClientData;
+		/** The Product User ID of the local user who is joining */
+		EOS_ProductUserId LocalUserId;
+		/** 
+		 * The UI Event associated with this Join Game event.
+		 * This should be used with EOS_Lobby_CopyLobbyDetailsHandleByUiEventId to get a handle to be used
+		 * when calling EOS_Lobby_JoinLobby.
+		 */
+		EOS_UI_EventId UiEventId;
+	));
+#endif
+
+	// Todo: handle UI events.
+}
+
+void FLobbiesEOS::RegisterHandlers()
+{
+	// Register for lobby updates.
+	OnLobbyUpdatedEOSEventRegistration = EOS_RegisterComponentEventHandler(
+		this,
+		LobbyPrerequisites->LobbyInterfaceHandle,
+		EOS_LOBBY_ADDNOTIFYLOBBYUPDATERECEIVED_API_LATEST,
+		&EOS_Lobby_AddNotifyLobbyUpdateReceived,
+		&EOS_Lobby_RemoveNotifyLobbyUpdateReceived,
+		&FLobbiesEOS::HandleLobbyUpdated);
+
+	// Register for lobby member updates.
+	OnLobbyMemberUpdatedEOSEventRegistration = EOS_RegisterComponentEventHandler(
+		this,
+		LobbyPrerequisites->LobbyInterfaceHandle,
+		EOS_LOBBY_ADDNOTIFYLOBBYMEMBERUPDATERECEIVED_API_LATEST,
+		&EOS_Lobby_AddNotifyLobbyMemberUpdateReceived,
+		&EOS_Lobby_RemoveNotifyLobbyMemberUpdateReceived,
+		&FLobbiesEOS::HandleLobbyMemberUpdated);
+
+	// Register for lobby member status changed.
+	OnLobbyMemberStatusReceivedEOSEventRegistration = EOS_RegisterComponentEventHandler(
+		this,
+		LobbyPrerequisites->LobbyInterfaceHandle,
+		EOS_LOBBY_ADDNOTIFYLOBBYMEMBERSTATUSRECEIVED_API_LATEST,
+		&EOS_Lobby_AddNotifyLobbyMemberStatusReceived,
+		&EOS_Lobby_RemoveNotifyLobbyMemberStatusReceived,
+		&FLobbiesEOS::HandleLobbyMemberStatusReceived);
+
+	// Register for lobby invite received.
+	OnLobbyInviteReceivedEOSEventRegistration = EOS_RegisterComponentEventHandler(
+		this,
+		LobbyPrerequisites->LobbyInterfaceHandle,
+		EOS_LOBBY_ADDNOTIFYLOBBYINVITERECEIVED_API_LATEST,
+		&EOS_Lobby_AddNotifyLobbyInviteReceived,
+		&EOS_Lobby_RemoveNotifyLobbyInviteReceived,
+		&FLobbiesEOS::HandleLobbyInviteReceived);
+
+	// Register for lobby invite accepted.
+	OnLobbyInviteAcceptedEOSEventRegistration = EOS_RegisterComponentEventHandler(
+		this,
+		LobbyPrerequisites->LobbyInterfaceHandle,
+		EOS_LOBBY_ADDNOTIFYLOBBYINVITERECEIVED_API_LATEST,
+		&EOS_Lobby_AddNotifyLobbyInviteAccepted,
+		&EOS_Lobby_RemoveNotifyLobbyInviteAccepted,
+		&FLobbiesEOS::HandleLobbyInviteAccepted);
+
+	// Register for join lobby accepted via overlay.
+	OnJoinLobbyAcceptedEOSEventRegistration = EOS_RegisterComponentEventHandler(
+		this,
+		LobbyPrerequisites->LobbyInterfaceHandle,
+		EOS_LOBBY_ADDNOTIFYJOINLOBBYACCEPTED_API_LATEST,
+		&EOS_Lobby_AddNotifyJoinLobbyAccepted,
+		&EOS_Lobby_RemoveNotifyJoinLobbyAccepted,
+		&FLobbiesEOS::HandleJoinLobbyAccepted);
+}
+
+void FLobbiesEOS::UnregisterHandlers()
+{
+	OnLobbyUpdatedEOSEventRegistration = nullptr;
+	OnLobbyMemberUpdatedEOSEventRegistration = nullptr;
+	OnLobbyMemberStatusReceivedEOSEventRegistration = nullptr;
+	OnLobbyInviteReceivedEOSEventRegistration = nullptr;
+	OnLobbyInviteAcceptedEOSEventRegistration = nullptr;
+	OnJoinLobbyAcceptedEOSEventRegistration = nullptr;
+}
+
+void FLobbiesEOS::AddActiveLobby(FOnlineAccountIdHandle LocalUserId, const TSharedRef<FLobbyDataEOS>& LobbyData)
+{
+	if (LobbyData->GetLocalMembers().IsEmpty())
+	{
+		// Notify lobby added.
+		OnLobbyJoinedEvent.Broadcast(FLobbyJoined{LobbyData->GetLobbyImpl()});
+	}
+
+	ActiveLobbies.FindOrAdd(LocalUserId).Add(LobbyData);
+	LobbyData->GetLocalMembers().Add(LocalUserId);
+
+	// Todo: make this better.
+	for (const TSharedRef<const FLobbyMember>& MemberData : LobbyData->GetLobbyImpl()->Members)
+	{
+		if (MemberData->AccountId == LocalUserId)
+		{
+			OnLobbyMemberJoinedEvent.Broadcast(FLobbyMemberJoined{LobbyData->GetLobbyImpl(), MemberData});
+			break;
+		}
+	}
+}
+
+void FLobbiesEOS::RemoveActiveLobby(FOnlineAccountIdHandle LocalUserId, const TSharedRef<FLobbyDataEOS>& LobbyData)
+{
+	// Todo: make this better.
+	for (const TSharedRef<const FLobbyMember>& MemberData : LobbyData->GetLobbyImpl()->Members)
+	{
+		if (MemberData->AccountId == LocalUserId)
+		{
+			OnLobbyMemberLeftEvent.Broadcast(FLobbyMemberLeft{LobbyData->GetLobbyImpl(), MemberData});
+			break;
+		}
+	}
+
+	ActiveLobbies.FindOrAdd(LocalUserId).Remove(LobbyData);
+	LobbyData->GetLocalMembers().Remove(LocalUserId);
+
+	if (LobbyData->GetLocalMembers().IsEmpty())
+	{
+		// Notify lobby removed.
+		OnLobbyLeftEvent.Broadcast(FLobbyLeft{LobbyData->GetLobbyImpl()});
+	}
+}
+
+void FLobbiesEOS::AddActiveInvite(const TSharedRef<FLobbyInviteDataEOS>& Invite)
+{
+	TMap<FOnlineLobbyIdHandle, TSharedRef<FLobbyInviteDataEOS>>& ActiveUserInvites = ActiveInvites.FindOrAdd(Invite->GetReceiver());
+	const FOnlineLobbyIdHandle LobbyId = Invite->GetLobbyData()->GetLobbyIdHandle();
+
+	// Todo: Handle multiple invites for the same lobby.
+	if (ActiveUserInvites.Find(LobbyId) == nullptr)
+	{
+		ActiveUserInvites.Add(LobbyId, Invite);
+		OnLobbyInvitationAddedEvent.Broadcast(
+			FLobbyInvitationAdded{
+				Invite->GetReceiver(),
+				Invite->GetSender(),
+				Invite->GetLobbyData()->GetLobbyImpl()
+			});
+	}
+}
+
+void FLobbiesEOS::RemoveActiveInvite(const TSharedRef<FLobbyInviteDataEOS>& Invite)
+{
+	ActiveInvites.FindOrAdd(Invite->GetReceiver()).Remove(Invite->GetLobbyData()->GetLobbyIdHandle());
+
+	OnLobbyInvitationRemovedEvent.Broadcast(
+		FLobbyInvitationRemoved{
+			Invite->GetReceiver(),
+			Invite->GetSender(),
+			Invite->GetLobbyData()->GetLobbyImpl()
+		});
+}
+
+TSharedPtr<FLobbyInviteDataEOS> FLobbiesEOS::GetActiveInvite(FOnlineAccountIdHandle TargetUser, FOnlineLobbyIdHandle TargetLobbyId)
+{
+	const TSharedRef<FLobbyInviteDataEOS>* Result = ActiveInvites.FindOrAdd(TargetUser).Find(TargetLobbyId);
+	return Result ? *Result : TSharedPtr<FLobbyInviteDataEOS>();
+}
+
+TFuture<TDefaultErrorResult<FLobbiesEOS::FJoinLobbyImpl>> FLobbiesEOS::JoinLobbyImpl(FJoinLobbyImpl::Params&& Params)
+{
+	// Check prerequisites.
+	if (!Services.Get<FAuthEOS>()->IsLoggedIn(Params.LocalUserId))
+	{
+		return MakeFulfilledPromise<TDefaultErrorResult<FJoinLobbyImpl>>(Errors::InvalidUser()).GetFuture();
+	}
+
+	if (!Params.LobbyData.IsValid())
+	{
+		return MakeFulfilledPromise<TDefaultErrorResult<FJoinLobbyImpl>>(Errors::InvalidParams()).GetFuture();
+	}
+
+	// Check whether any of the local users is already in the target lobby.
+	for (const FJoinLobbyLocalUserData& JoinData : Params.LocalUsers)
+	{
+		if (const TSet<TSharedRef<FLobbyDataEOS>>* UserActiveLobbies = ActiveLobbies.Find(JoinData.LocalUserId))
+		{
+			if (UserActiveLobbies->Find(Params.LobbyData.ToSharedRef()))
+			{
+				return MakeFulfilledPromise<TDefaultErrorResult<FJoinLobbyImpl>>(Errors::InvalidParams()).GetFuture();
+			}
+		}
+	}
+
+	TArray<TFuture<TDefaultErrorResult<FJoinLobbyMemberImpl>>> PendingMemberJoins;
+	for (FJoinLobbyLocalUserData& UserData : Params.LocalUsers)
+	{
+		FJoinLobbyMemberImpl::Params JoinLobbyMemberParams;
+		JoinLobbyMemberParams.LobbyData = Params.LobbyData.ToSharedRef();
+		JoinLobbyMemberParams.LocalUserId = UserData.LocalUserId;
+		JoinLobbyMemberParams.Attributes = MoveTemp(UserData.Attributes);
+
+		auto JoinMemberPromise = MakeShared<TPromise<TDefaultErrorResult<FJoinLobbyMemberImpl>>>();
+		JoinLobbyMemberImplOp(MoveTemp(JoinLobbyMemberParams))
+		.OnComplete([JoinMemberPromise](const TOnlineResult<FJoinLobbyMemberImpl>& Result)
+		{
+			if (Result.IsOk())
+			{
+				JoinMemberPromise->EmplaceValue(Result.GetOkValue());
+			}
+			else
+			{
+				JoinMemberPromise->EmplaceValue(Result.GetErrorValue());
+			}
+		});
+
+		PendingMemberJoins.Emplace(JoinMemberPromise->GetFuture());
+	}
+
+	TPromise<TDefaultErrorResult<FJoinLobbyImpl>> Promise;
+	auto Future = Promise.GetFuture();
+
+	WhenAll(MoveTemp(PendingMemberJoins))
+	.Then([this, Promise = MoveTemp(Promise), Params = MoveTemp(Params)](TFuture<TArray<TDefaultErrorResult<FJoinLobbyMemberImpl>>>&& Results) mutable
+	{
+		TOptional<FOnlineError> StoredError;
+		for (TDefaultErrorResult<FJoinLobbyMemberImpl>& Result : Results.Get())
+		{
+			if (Result.IsError())
+			{
+				// Store first encountered error to return as result.
+				StoredError = Result.GetErrorValue();
+			}
+		}
+
+		if (StoredError)
+		{
+			TArray<TFuture<TDefaultErrorResult<FLeaveLobbyImpl>>> PendingMemberExits;
+			for (int32 MemberIndex = 0; MemberIndex < Params.LocalUsers.Num(); ++MemberIndex)
+			{
+				const TDefaultErrorResult<FJoinLobbyMemberImpl>& MemberResult = Results.Get()[MemberIndex];
+				const FOnlineAccountIdHandle MemberId = Params.LocalUsers[MemberIndex].LocalUserId;
+
+				if (Results.Get()[MemberIndex].IsError())
+				{
+					FLeaveLobbyImpl::Params LeaveLobbyParams;
+					LeaveLobbyParams.LobbyData = Params.LobbyData.ToSharedRef();
+					LeaveLobbyParams.LocalUserId = MemberId;
+
+					TPromise<TDefaultErrorResult<FLeaveLobbyImpl>> LeaveMemberPromise;
+					PendingMemberExits.Emplace(LeaveMemberPromise.GetFuture());
+
+					LeaveLobbyImpl(MoveTemp(LeaveLobbyParams))
+					.Then([LeaveMemberPromise = MoveTemp(LeaveMemberPromise)](TFuture<TDefaultErrorResult<FLeaveLobbyImpl>>&& Future) mutable
+					{
+						if (Future.Get().IsError())
+						{
+							// Todo: complain about having an error while handling an error.
+						}
+
+						LeaveMemberPromise.EmplaceValue(MoveTempIfPossible(Future.Get()));
+					});
+				}
+			}
+
+			WhenAll(MoveTemp(PendingMemberExits))
+			.Then([Promise = MoveTemp(Promise), StoredError = MoveTemp(*StoredError)](TFuture<TArray<TDefaultErrorResult<FLeaveLobbyImpl>>>&& Future) mutable
+			{
+				Promise.EmplaceValue(MoveTemp(StoredError));
+			});
+		}
+		else
+		{
+			Promise.EmplaceValue(FJoinLobbyImpl::Result{});
+		}
+	});
+
+	return Future;
+}
+
+TOnlineAsyncOpHandle<FLobbiesEOS::FJoinLobbyMemberImpl> FLobbiesEOS::JoinLobbyMemberImplOp(FJoinLobbyMemberImpl::Params&& InParams)
+{
+	TOnlineAsyncOpRef<FJoinLobbyMemberImpl> Op = GetOp<FJoinLobbyMemberImpl>(MoveTemp(InParams));
+	const FJoinLobbyMemberImpl::Params& Params = Op->GetParams();
+
+	// Setup lobby details - Prefer UI event before invitation before search result.
+	TSharedPtr<FLobbyDetailsEOS> LobbyDetails = Params.LobbyData->GetUserLobbyDetails(Params.LocalUserId);
+	if (!LobbyDetails)
+	{
+		// Todo: Check whether another local member can invite the user.
+		Op->SetError(Errors::InvalidParams());
+		return Op->GetHandle();
+	}
+
+	Op->Then([this, LobbyDetails, LobbyData = Params.LobbyData](TOnlineAsyncOp<FJoinLobbyMemberImpl>& InAsyncOp)
+	{
+		const auto& Params = InAsyncOp.GetParams();
+
+		EOS_Lobby_JoinLobbyOptions JoinLobbyOptions = {};
+		JoinLobbyOptions.ApiVersion = EOS_LOBBY_JOINLOBBY_API_LATEST;
+		JoinLobbyOptions.LobbyDetailsHandle = LobbyDetails->GetEOSHandle();
+		JoinLobbyOptions.LocalUserId = GetProductUserIdChecked(Params.LocalUserId);
+		JoinLobbyOptions.bPresenceEnabled = false;
+		JoinLobbyOptions.LocalRTCOptions = nullptr;
+
+		UE_LOG(LogLobbies, Warning, TEXT("[FLobbiesEOS::EOS_Lobby_JoinLobby] Start. Lobby: %s, Member: %s"),
+			*LobbyData->GetLobbyId(), *ToLogString(Params.LocalUserId));
+
+		return EOS_Async<EOS_Lobby_JoinLobbyCallbackInfo>(EOS_Lobby_JoinLobby, LobbyPrerequisites->LobbyInterfaceHandle, JoinLobbyOptions);
+	})
+	.Then([this](TOnlineAsyncOp<FJoinLobbyMemberImpl>& InAsyncOp, const EOS_Lobby_JoinLobbyCallbackInfo* Data)
+	{
+		UE_LOG(LogLobbies, Warning, TEXT("[FLobbiesEOS::EOS_Lobby_JoinLobby] Complete. Result: %d"), Data->ResultCode);
+		if (Data->ResultCode != EOS_EResult::EOS_Success)
+		{
+			// TODO: Error codes
+			InAsyncOp.SetError(Errors::Unknown(FromEOSError(Data->ResultCode)));
+		}
+	})
+	.Then([this](TOnlineAsyncOp<FJoinLobbyMemberImpl>& InAsyncOp)
+	{
+		const auto& Params = InAsyncOp.GetParams();
+		FModifyLobbyMemberDataImpl::Params ModifyParams;
+		ModifyParams.LobbyData = Params.LobbyData;
+		ModifyParams.LocalUserId = Params.LocalUserId;
+		ModifyParams.Changes = MakeShared<FLobbyMemberDataChanges>();
+		ModifyParams.Changes->MutatedAttributes = Params.Attributes;
+		return ModifyLobbyMemberDataImpl(MoveTemp(ModifyParams));
+	})
+	.Then([this](TOnlineAsyncOp<FJoinLobbyMemberImpl>& InAsyncOp, TDefaultErrorResult<FModifyLobbyMemberDataImpl>&& Result)
+	{
+		const auto& Params = InAsyncOp.GetParams();
+
+		if (Result.IsError())
+		{
+			// Failed to set attributes - leave the lobby.
+			FLeaveLobbyImpl::Params LeaveLobbyParams;
+			LeaveLobbyParams.LobbyData = Params.LobbyData;
+			LeaveLobbyParams.LocalUserId = Params.LocalUserId;
+
+			TPromise<void> Promise;
+			auto Future = Promise.GetFuture();
+
+			LeaveLobbyImpl(MoveTemp(LeaveLobbyParams))
+			.Then([this, InAsyncOp = InAsyncOp.AsShared(), ErrorResult = MoveTemp(Result.GetErrorValue()), Promise = MoveTemp(Promise)](TFuture<TDefaultErrorResult<FLeaveLobbyImpl>>&& Future) mutable
+			{
+				if (Future.Get().IsError())
+				{
+					// Todo: complain about having an error while handling an error.
+				}
+
+				// Todo: Errors.
+				InAsyncOp->SetError(Errors::Unknown(MoveTemp(ErrorResult)));
+				Promise.EmplaceValue();
+			});
+
+			return Future;
+		}
+		else
+		{
+			return MakeFulfilledPromise<void>().GetFuture();
+		}
+	})
+	.Then([this](TOnlineAsyncOp<FJoinLobbyMemberImpl>& InAsyncOp)
+	{
+		InAsyncOp.SetResult(FJoinLobbyMemberImpl::Result{});
+	})
+	.Enqueue(GetSerialQueue(Params.LocalUserId));
+
+	return Op->GetHandle();
+}
+
+TFuture<TDefaultErrorResult<FLobbiesEOS::FLeaveLobbyImpl>> FLobbiesEOS::LeaveLobbyImpl(FLeaveLobbyImpl::Params&& Params)
+{
+	// Check prerequisites.
+	if (!Services.Get<FAuthEOS>()->IsLoggedIn(Params.LocalUserId))
+	{
+		return MakeFulfilledPromise<TDefaultErrorResult<FLeaveLobbyImpl>>(Errors::InvalidUser()).GetFuture();
+	}
+
+	if (!Params.LobbyData.IsValid())
+	{
+		return MakeFulfilledPromise<TDefaultErrorResult<FLeaveLobbyImpl>>(Errors::InvalidParams()).GetFuture();
+	}
+
+	EOS_Lobby_LeaveLobbyOptions LeaveLobbyOptions = {};
+	LeaveLobbyOptions.ApiVersion = EOS_LOBBY_LEAVELOBBY_API_LATEST;
+	LeaveLobbyOptions.LocalUserId = GetProductUserIdChecked(Params.LocalUserId);
+	LeaveLobbyOptions.LobbyId = Params.LobbyData->GetLobbyIdEOS();
+
+	TPromise<TDefaultErrorResult<FLeaveLobbyImpl>> Promise;
+	auto Future = Promise.GetFuture();
+
+	EOS_Async<EOS_Lobby_LeaveLobbyCallbackInfo>(EOS_Lobby_LeaveLobby, LobbyPrerequisites->LobbyInterfaceHandle, LeaveLobbyOptions)
+	.Then([Promise = MoveTemp(Promise)](TFuture<const EOS_Lobby_LeaveLobbyCallbackInfo*>&& Future) mutable
+	{
+		const EOS_Lobby_LeaveLobbyCallbackInfo* Result = Future.Get();
+		if (Result->ResultCode != EOS_EResult::EOS_Success)
+		{
+			// Todo: Errors
+			Promise.EmplaceValue(Errors::Unknown(FromEOSError(Result->ResultCode)));
+			return;
+		}
+
+		Promise.EmplaceValue(FLeaveLobbyImpl::Result{});
+	});
+
+	return Future;
+}
+
+TFuture<TDefaultErrorResult<FLobbiesEOS::FDestroyLobbyImpl>> FLobbiesEOS::DestroyLobbyImpl(FDestroyLobbyImpl::Params&& Params)
+{
+	// Check prerequisites.
+	if (!Services.Get<FAuthEOS>()->IsLoggedIn(Params.LocalUserId))
+	{
+		return MakeFulfilledPromise<TDefaultErrorResult<FDestroyLobbyImpl>>(Errors::InvalidUser()).GetFuture();
+	}
+
+	if (!Params.LobbyData.IsValid())
+	{
+		return MakeFulfilledPromise<TDefaultErrorResult<FDestroyLobbyImpl>>(Errors::InvalidParams()).GetFuture();
+	}
+
+	EOS_Lobby_DestroyLobbyOptions DestroyLobbyOptions = {};
+	DestroyLobbyOptions.ApiVersion = EOS_LOBBY_DESTROYLOBBY_API_LATEST;
+	DestroyLobbyOptions.LocalUserId = GetProductUserIdChecked(Params.LocalUserId);
+	DestroyLobbyOptions.LobbyId = Params.LobbyData->GetLobbyIdEOS();
+
+	TPromise<TDefaultErrorResult<FDestroyLobbyImpl>> Promise;
+	auto Future = Promise.GetFuture();
+
+	EOS_Async<EOS_Lobby_DestroyLobbyCallbackInfo>(EOS_Lobby_DestroyLobby, LobbyPrerequisites->LobbyInterfaceHandle, DestroyLobbyOptions)
+	.Then([Promise = MoveTemp(Promise)](TFuture<const EOS_Lobby_DestroyLobbyCallbackInfo*>&& Future) mutable
+	{
+		const EOS_Lobby_DestroyLobbyCallbackInfo* Result = Future.Get();
+		if (Result->ResultCode != EOS_EResult::EOS_Success)
+		{
+			// Todo: Errors
+			Promise.EmplaceValue(Errors::Unknown(FromEOSError(Result->ResultCode)));
+			return;
+		}
+
+		Promise.EmplaceValue(FDestroyLobbyImpl::Result{});
+	});
+
+	return Future;
+}
+
+TFuture<TDefaultErrorResult<FLobbiesEOS::FInviteLobbyMemberImpl>> FLobbiesEOS::InviteLobbyMemberImpl(FInviteLobbyMemberImpl::Params&& Params)
+{
+	// Check prerequisites.
+	if (!Services.Get<FAuthEOS>()->IsLoggedIn(Params.LocalUserId))
+	{
+		return MakeFulfilledPromise<TDefaultErrorResult<FInviteLobbyMemberImpl>>(Errors::InvalidUser()).GetFuture();
+	}
+
+	if (!Params.LobbyData.IsValid())
+	{
+		return MakeFulfilledPromise<TDefaultErrorResult<FInviteLobbyMemberImpl>>(Errors::InvalidParams()).GetFuture();
+	}
+
+	EOS_Lobby_SendInviteOptions SendInviteOptions = {};
+	SendInviteOptions.ApiVersion = EOS_LOBBY_SENDINVITE_API_LATEST;
+	SendInviteOptions.LocalUserId = GetProductUserIdChecked(Params.LocalUserId);
+	SendInviteOptions.TargetUserId = GetProductUserIdChecked(Params.TargetUserId);
+	SendInviteOptions.LobbyId = Params.LobbyData->GetLobbyIdEOS();
+
+	TPromise<TDefaultErrorResult<FInviteLobbyMemberImpl>> Promise;
+	auto Future = Promise.GetFuture();
+
+	EOS_Async<EOS_Lobby_SendInviteCallbackInfo>(EOS_Lobby_SendInvite, LobbyPrerequisites->LobbyInterfaceHandle, SendInviteOptions)
+	.Then([Promise = MoveTemp(Promise)](TFuture<const EOS_Lobby_SendInviteCallbackInfo*>&& Future) mutable
+	{
+		const EOS_Lobby_SendInviteCallbackInfo* Result = Future.Get();
+		if (Result->ResultCode != EOS_EResult::EOS_Success)
+		{
+			// Todo: Errors
+			Promise.EmplaceValue(Errors::Unknown(FromEOSError(Result->ResultCode)));
+			return;
+		}
+
+		Promise.EmplaceValue(FInviteLobbyMemberImpl::Result{});
+	});
+
+	return Future;
+}
+
+TFuture<TDefaultErrorResult<FLobbiesEOS::FDeclineLobbyInvitationImpl>> FLobbiesEOS::DeclineLobbyInvitationImpl(FDeclineLobbyInvitationImpl::Params&& Params)
+{
+	// Check prerequisites.
+	if (!Services.Get<FAuthEOS>()->IsLoggedIn(Params.LocalUserId))
+	{
+		return MakeFulfilledPromise<TDefaultErrorResult<FDeclineLobbyInvitationImpl>>(Errors::InvalidUser()).GetFuture();
+	}
+
+	// Find the active invitation.
+	TSharedPtr<FLobbyInviteDataEOS> InviteData = GetActiveInvite(Params.LocalUserId, Params.LobbyId);
+	if (!InviteData)
+	{
+		// Todo: Errors.
+		return MakeFulfilledPromise<TDefaultErrorResult<FDeclineLobbyInvitationImpl>>(Errors::Unknown()).GetFuture();
+	}
+
+	EOS_Lobby_RejectInviteOptions RejectInviteOptions = {};
+	RejectInviteOptions.ApiVersion = EOS_LOBBY_REJECTINVITE_API_LATEST;
+	RejectInviteOptions.LocalUserId = GetProductUserIdChecked(Params.LocalUserId);
+	RejectInviteOptions.InviteId = InviteData->GetInviteIdEOS();
+
+	TPromise<TDefaultErrorResult<FDeclineLobbyInvitationImpl>> Promise;
+	TFuture<TDefaultErrorResult<FDeclineLobbyInvitationImpl>> Future = Promise.GetFuture();
+
+	EOS_Async<EOS_Lobby_RejectInviteCallbackInfo>(EOS_Lobby_RejectInvite, LobbyPrerequisites->LobbyInterfaceHandle, RejectInviteOptions)
+	.Then([this, Promise = MoveTemp(Promise), InviteData](TFuture<const EOS_Lobby_RejectInviteCallbackInfo*>&& Future) mutable
+	{
+		// Remove active invitation.
+		RemoveActiveInvite(InviteData.ToSharedRef());
+
+		const EOS_Lobby_RejectInviteCallbackInfo* Result = Future.Get();
+		if (Result->ResultCode != EOS_EResult::EOS_Success)
+		{
+			// Todo: Errors
+			Promise.EmplaceValue(Errors::Unknown(FromEOSError(Result->ResultCode)));
+			return;
+		}
+
+		Promise.EmplaceValue(FDeclineLobbyInvitationImpl::Result{});
+	});
+
+	return Future;
+}
+
+TFuture<TDefaultErrorResult<FLobbiesEOS::FKickLobbyMemberImpl>> FLobbiesEOS::KickLobbyMemberImpl(FKickLobbyMemberImpl::Params&& Params)
+{
+	// Check prerequisites.
+	if (!Services.Get<FAuthEOS>()->IsLoggedIn(Params.LocalUserId))
+	{
+		return MakeFulfilledPromise<TDefaultErrorResult<FKickLobbyMemberImpl>>(Errors::InvalidUser()).GetFuture();
+	}
+
+	if (!Params.LobbyData.IsValid())
+	{
+		return MakeFulfilledPromise<TDefaultErrorResult<FKickLobbyMemberImpl>>(Errors::InvalidParams()).GetFuture();
+	}
+
+	// todo: check local user is lobby owner
+
+	EOS_Lobby_KickMemberOptions KickMemberOptions = {};
+	KickMemberOptions.ApiVersion = EOS_LOBBY_KICKMEMBER_API_LATEST;
+	KickMemberOptions.LocalUserId = GetProductUserIdChecked(Params.LocalUserId);
+	KickMemberOptions.TargetUserId = GetProductUserIdChecked(Params.TargetUserId);
+	KickMemberOptions.LobbyId = Params.LobbyData->GetLobbyIdEOS();
+
+	TPromise<TDefaultErrorResult<FKickLobbyMemberImpl>> Promise;
+	TFuture<TDefaultErrorResult<FKickLobbyMemberImpl>> Future = Promise.GetFuture();
+
+	EOS_Async<EOS_Lobby_KickMemberCallbackInfo>(EOS_Lobby_KickMember, LobbyPrerequisites->LobbyInterfaceHandle, KickMemberOptions)
+	.Then([Promise = MoveTemp(Promise)](TFuture<const EOS_Lobby_KickMemberCallbackInfo*>&& Future) mutable
+	{
+		const EOS_Lobby_KickMemberCallbackInfo* Result = Future.Get();
+		if (Result->ResultCode != EOS_EResult::EOS_Success)
+		{
+			// Todo: Errors
+			Promise.EmplaceValue(Errors::Unknown(FromEOSError(Result->ResultCode)));
+			return;
+		}
+
+		Promise.EmplaceValue(FKickLobbyMemberImpl::Result{});
+	});
+
+	return Future;
+}
+
+TFuture<TDefaultErrorResult<FLobbiesEOS::FPromoteLobbyMemberImpl>> FLobbiesEOS::PromoteLobbyMemberImpl(FPromoteLobbyMemberImpl::Params&& Params)
+{
+	// Check prerequisites.
+	if (!Services.Get<FAuthEOS>()->IsLoggedIn(Params.LocalUserId))
+	{
+		return MakeFulfilledPromise<TDefaultErrorResult<FPromoteLobbyMemberImpl>>(Errors::InvalidUser()).GetFuture();
+	}
+
+	if (!Params.LobbyData.IsValid())
+	{
+		return MakeFulfilledPromise<TDefaultErrorResult<FPromoteLobbyMemberImpl>>(Errors::InvalidParams()).GetFuture();
+	}
+
+	// todo: check local user is lobby owner
+
+	EOS_Lobby_PromoteMemberOptions PromoteMemberOptions = {};
+	PromoteMemberOptions.ApiVersion = EOS_LOBBY_PROMOTEMEMBER_API_LATEST;
+	PromoteMemberOptions.LocalUserId = GetProductUserIdChecked(Params.LocalUserId);
+	PromoteMemberOptions.TargetUserId = GetProductUserIdChecked(Params.TargetUserId);
+	PromoteMemberOptions.LobbyId = Params.LobbyData->GetLobbyIdEOS();
+
+	TPromise<TDefaultErrorResult<FPromoteLobbyMemberImpl>> Promise;
+	TFuture<TDefaultErrorResult<FPromoteLobbyMemberImpl>> Future = Promise.GetFuture();
+
+	EOS_Async<EOS_Lobby_PromoteMemberCallbackInfo>(EOS_Lobby_PromoteMember, LobbyPrerequisites->LobbyInterfaceHandle, PromoteMemberOptions)
+	.Then([Promise = MoveTemp(Promise)](TFuture<const EOS_Lobby_PromoteMemberCallbackInfo*>&& Future) mutable
+	{
+		const EOS_Lobby_PromoteMemberCallbackInfo* Result = Future.Get();
+		if (Result->ResultCode != EOS_EResult::EOS_Success)
+		{
+			// Todo: Errors
+			Promise.EmplaceValue(Errors::Unknown(FromEOSError(Result->ResultCode)));
+			return;
+		}
+
+		Promise.EmplaceValue(FPromoteLobbyMemberImpl::Result{});
+	});
+
+	return Future;
+}
+
+TFuture<TDefaultErrorResult<FLobbiesEOS::FModifyLobbyDataImpl>> FLobbiesEOS::ModifyLobbyDataImpl(FModifyLobbyDataImpl::Params&& Params)
+{
+	// Check prerequisites.
+	if (!Services.Get<FAuthEOS>()->IsLoggedIn(Params.LocalUserId))
+	{
+		return MakeFulfilledPromise<TDefaultErrorResult<FModifyLobbyDataImpl>>(Errors::InvalidUser()).GetFuture();
+	}
+
+	if (!Params.LobbyData.IsValid())
+	{
+		return MakeFulfilledPromise<TDefaultErrorResult<FModifyLobbyDataImpl>>(Errors::InvalidParams()).GetFuture();
+	}
+
+	TDefaultErrorResultInternal<TSharedPtr<FLobbyDetailsEOS>> DetailsResult = FLobbyDetailsEOS::CreateFromLobbyId(
+		LobbyPrerequisites.ToSharedRef(), Params.LobbyData->GetLobbyIdEOS(), Params.LocalUserId);
+	if (DetailsResult.IsError())
+	{
+		// Todo: Errors
+		return MakeFulfilledPromise<TDefaultErrorResult<FModifyLobbyDataImpl>>(Errors::Unknown(MoveTemp(DetailsResult.GetErrorValue()))).GetFuture();
+	}
+
+	// todo: check local user is lobby owner
+
+	UE_LOG(LogLobbies, Warning, TEXT("[FLobbiesEOS::ModifyLobbyDataImpl] Start. Lobby: %s, Member: %s"),
+		*Params.LobbyData->GetLobbyId(), *ToLogString(Params.LocalUserId));
+
+	TPromise<TDefaultErrorResult<FModifyLobbyDataImpl>> Promise;
+	TFuture<TDefaultErrorResult<FModifyLobbyDataImpl>> Future = Promise.GetFuture();
+
+	DetailsResult.GetOkValue()->ApplyLobbyDataUpdates(Params.LocalUserId, MoveTemp(*Params.Changes))
+	.Then([Promise = MoveTemp(Promise)](TFuture<EOS_EResult>&& Future) mutable
+	{
+		UE_LOG(LogLobbies, Warning, TEXT("[FLobbiesEOS::ModifyLobbyDataImpl] Complete. Result: %d"), Future.Get());
+
+		// Todo: Handle "no change" better.
+		if (Future.Get() != EOS_EResult::EOS_Success && Future.Get() != EOS_EResult::EOS_NoChange)
+		{
+			// Todo: Errors
+			Promise.EmplaceValue(Errors::Unknown(FromEOSError(Future.Get())));
+			return;
+		}
+
+		Promise.EmplaceValue(FModifyLobbyDataImpl::Result{});
+	});
+
+	return Future;
+}
+
+TFuture<TDefaultErrorResult<FLobbiesEOS::FModifyLobbyMemberDataImpl>> FLobbiesEOS::ModifyLobbyMemberDataImpl(FModifyLobbyMemberDataImpl::Params&& Params)
+{
+	// Check prerequisites.
+	if (!Services.Get<FAuthEOS>()->IsLoggedIn(Params.LocalUserId))
+	{
+		return MakeFulfilledPromise<TDefaultErrorResult<FModifyLobbyMemberDataImpl>>(Errors::InvalidUser()).GetFuture();
+	}
+
+	if (!Params.LobbyData.IsValid())
+	{
+		return MakeFulfilledPromise<TDefaultErrorResult<FModifyLobbyMemberDataImpl>>(Errors::InvalidParams()).GetFuture();
+	}
+
+	TDefaultErrorResultInternal<TSharedPtr<FLobbyDetailsEOS>> DetailsResult = FLobbyDetailsEOS::CreateFromLobbyId(
+		LobbyPrerequisites.ToSharedRef(), Params.LobbyData->GetLobbyIdEOS(), Params.LocalUserId);
+	if (DetailsResult.IsError())
+	{
+		// Todo: Errors
+		return MakeFulfilledPromise<TDefaultErrorResult<FModifyLobbyMemberDataImpl>>(Errors::Unknown(MoveTemp(DetailsResult.GetErrorValue()))).GetFuture();
+	}
+
+	UE_LOG(LogLobbies, Warning, TEXT("[FLobbiesEOS::ModifyLobbyMemberDataImpl] Start. Lobby: %s, Member: %s"),
+		*Params.LobbyData->GetLobbyId(), *ToLogString(Params.LocalUserId));
+
+	TPromise<TDefaultErrorResult<FModifyLobbyMemberDataImpl>> Promise;
+	TFuture<TDefaultErrorResult<FModifyLobbyMemberDataImpl>> Future = Promise.GetFuture();
+
+	DetailsResult.GetOkValue()->ApplyLobbyMemberDataUpdates(Params.LocalUserId, MoveTemp(*Params.Changes))
+	.Then([Promise = MoveTemp(Promise)](TFuture<EOS_EResult>&& Future) mutable
+	{
+		UE_LOG(LogLobbies, Warning, TEXT("[FLobbiesEOS::ModifyLobbyMemberDataImpl] Complete. Result: %d"), Future.Get());
+
+		// Todo: Handle "no change" better.
+		if (Future.Get() != EOS_EResult::EOS_Success && Future.Get() != EOS_EResult::EOS_NoChange)
+		{
+			// Todo: Errors
+			Promise.EmplaceValue(Errors::Unknown(FromEOSError(Future.Get())));
+			return;
+		}
+
+		Promise.EmplaceValue(FModifyLobbyMemberDataImpl::Result{});
+	});
+
+	return Future;
+}
+
+/* UE::Online */ }
diff --git a/Engine/Plugins/Online/Experimental/OnlineServicesEOS/Source/Private/LobbiesEOSTypes.cpp b/Engine/Plugins/Online/Experimental/OnlineServicesEOS/Source/Private/LobbiesEOSTypes.cpp
new file mode 100644
--- /dev/null
+++ b/Engine/Plugins/Online/Experimental/OnlineServicesEOS/Source/Private/LobbiesEOSTypes.cpp
@@ -0,0 +1,1015 @@
+// Copyright Epic Games, Inc. All Rights Reserved.
+
+#include "LobbiesEOSTypes.h"
+#include "AuthEOS.h"
+#include "Online/LobbiesCommon.h"
+
+namespace UE::Online {
+
+namespace detail {
+
+FString TranslateLobbyId(EOS_LobbyId EOSLobbyId)
+{
+	return FString(UTF8_TO_TCHAR(EOSLobbyId));
+}
+
+FString TranslateLobbyInviteId(const char* InviteId)
+{
+	return FString(UTF8_TO_TCHAR(InviteId));
+}
+
+FLobbyAttributeId TranslateLobbyAttributeId(const char* AttributeId)
+{
+	return FLobbyAttributeId(UTF8_TO_TCHAR(AttributeId));
+}
+
+EOS_EComparisonOp TranslateSearchComparison(ELobbyComparisonOp Op)
+{
+	switch (Op)
+	{
+	default:									checkNoEntry(); // Intentional fallthrough
+	case ELobbyComparisonOp::Equals:			return EOS_EComparisonOp::EOS_CO_EQUAL;
+	case ELobbyComparisonOp::NotEquals:			return EOS_EComparisonOp::EOS_CO_NOTEQUAL;
+	case ELobbyComparisonOp::GreaterThan:		return EOS_EComparisonOp::EOS_CO_GREATERTHAN;
+	case ELobbyComparisonOp::GreaterThanEquals:	return EOS_EComparisonOp::EOS_CO_GREATERTHANOREQUAL;
+	case ELobbyComparisonOp::LessThan:			return EOS_EComparisonOp::EOS_CO_LESSTHAN;
+	case ELobbyComparisonOp::LessThanEquals:	return EOS_EComparisonOp::EOS_CO_LESSTHANOREQUAL;
+	case ELobbyComparisonOp::Near:				return EOS_EComparisonOp::EOS_CO_DISTANCE;
+	case ELobbyComparisonOp::In:				return EOS_EComparisonOp::EOS_CO_ONEOF;
+	case ELobbyComparisonOp::NotIn:				return EOS_EComparisonOp::EOS_CO_NOTANYOF;
+
+	// todo:
+	// EOS_EComparisonOp::EOS_CO_ANYOF
+	// EOS_EComparisonOp::EOS_CO_NOTONEOF
+	// EOS_EComparisonOp::EOS_CO_CONTAINS
+	}
+}
+
+class FLobbyNotificationPauseHandleImpl final : public FLobbyNotificationPauseHandle
+{
+public:
+	using FUnregisterFn = TFunction<void()>;
+
+	FLobbyNotificationPauseHandleImpl(FUnregisterFn&& UnregisterFn)
+		: UnregisterFn(MoveTemp(UnregisterFn))
+	{
+	}
+
+	virtual ~FLobbyNotificationPauseHandleImpl()
+	{
+		if (UnregisterFn)
+		{
+			UnregisterFn();
+		}
+	}
+
+private:
+	FUnregisterFn UnregisterFn;
+};
+
+} // detail
+
+const EOS_HLobbyDetails FLobbyDetailsEOS::InvalidLobbyDetailsHandle = {};
+
+TDefaultErrorResultInternal<TSharedPtr<FLobbyDetailsEOS>> FLobbyDetailsEOS::CreateFromLobbyId(const TSharedRef<FLobbyPrerequisitesEOS>& Prerequisites, EOS_LobbyId LobbyId, FOnlineAccountIdHandle MemberHandle)
+{
+	EOS_HLobbyDetails LobbyDetailsHandle = {};
+
+	EOS_Lobby_CopyLobbyDetailsHandleOptions Options;
+	Options.ApiVersion = EOS_LOBBY_COPYLOBBYDETAILSHANDLE_API_LATEST;
+	Options.LobbyId = LobbyId;
+	Options.LocalUserId = GetProductUserIdChecked(MemberHandle);
+
+	EOS_EResult EOSResult = EOS_Lobby_CopyLobbyDetailsHandle(Prerequisites->LobbyInterfaceHandle, &Options, &LobbyDetailsHandle);
+	if (EOSResult != EOS_EResult::EOS_Success)
+	{
+		return TDefaultErrorResultInternal<TSharedPtr<FLobbyDetailsEOS>>(FromEOSError(EOSResult));
+	}
+
+	TDefaultErrorResultInternal<TSharedPtr<FLobbyDetailsInfoEOS>> LobbyDetailsInfoResult = FLobbyDetailsInfoEOS::Create(LobbyDetailsHandle);
+	if (LobbyDetailsInfoResult.IsError())
+	{
+		EOS_LobbyDetails_Release(LobbyDetailsHandle);
+		return TDefaultErrorResultInternal<TSharedPtr<FLobbyDetailsEOS>>(MoveTemp(LobbyDetailsInfoResult.GetErrorValue()));
+	}
+
+	return TDefaultErrorResultInternal<TSharedPtr<FLobbyDetailsEOS>>(MakeShared<FLobbyDetailsEOS>(Prerequisites, LobbyDetailsInfoResult.GetOkValue().ToSharedRef(), ELobbyDetailsSource::Active, LobbyDetailsHandle));
+}
+
+TDefaultErrorResultInternal<TSharedPtr<FLobbyDetailsEOS>> FLobbyDetailsEOS::CreateFromInviteId(const TSharedRef<FLobbyPrerequisitesEOS>& Prerequisites, const char* InviteId)
+{
+	EOS_HLobbyDetails LobbyDetailsHandle = {};
+
+	EOS_Lobby_CopyLobbyDetailsHandleByInviteIdOptions Options;
+	Options.ApiVersion = EOS_LOBBY_COPYLOBBYDETAILSHANDLEBYINVITEID_API_LATEST;
+	Options.InviteId = InviteId;
+
+	EOS_EResult EOSResult = EOS_Lobby_CopyLobbyDetailsHandleByInviteId(Prerequisites->LobbyInterfaceHandle, &Options, &LobbyDetailsHandle);
+	if (EOSResult != EOS_EResult::EOS_Success)
+	{
+		return TDefaultErrorResultInternal<TSharedPtr<FLobbyDetailsEOS>>(FromEOSError(EOSResult));
+	}
+
+	TDefaultErrorResultInternal<TSharedPtr<FLobbyDetailsInfoEOS>> LobbyDetailsInfoResult = FLobbyDetailsInfoEOS::Create(LobbyDetailsHandle);
+	if (LobbyDetailsInfoResult.IsError())
+	{
+		EOS_LobbyDetails_Release(LobbyDetailsHandle);
+		return TDefaultErrorResultInternal<TSharedPtr<FLobbyDetailsEOS>>(MoveTemp(LobbyDetailsInfoResult.GetErrorValue()));
+	}
+
+	return TDefaultErrorResultInternal<TSharedPtr<FLobbyDetailsEOS>>(MakeShared<FLobbyDetailsEOS>(Prerequisites, LobbyDetailsInfoResult.GetOkValue().ToSharedRef(), ELobbyDetailsSource::Invite, LobbyDetailsHandle));
+}
+
+TDefaultErrorResultInternal<TSharedPtr<FLobbyDetailsEOS>> FLobbyDetailsEOS::CreateFromUiEventId(const TSharedRef<FLobbyPrerequisitesEOS>& Prerequisites, EOS_UI_EventId UiEventId)
+{
+	EOS_HLobbyDetails LobbyDetailsHandle = {};
+
+	EOS_Lobby_CopyLobbyDetailsHandleByUiEventIdOptions Options;
+	Options.ApiVersion = EOS_LOBBY_COPYLOBBYDETAILSHANDLEBYUIEVENTID_API_LATEST;
+	Options.UiEventId = UiEventId;
+
+	EOS_EResult EOSResult = EOS_Lobby_CopyLobbyDetailsHandleByUiEventId(Prerequisites->LobbyInterfaceHandle, &Options, &LobbyDetailsHandle);
+	if (EOSResult != EOS_EResult::EOS_Success)
+	{
+		return TDefaultErrorResultInternal<TSharedPtr<FLobbyDetailsEOS>>(FromEOSError(EOSResult));
+	}
+
+	TDefaultErrorResultInternal<TSharedPtr<FLobbyDetailsInfoEOS>> LobbyDetailsInfoResult = FLobbyDetailsInfoEOS::Create(LobbyDetailsHandle);
+	if (LobbyDetailsInfoResult.IsError())
+	{
+		EOS_LobbyDetails_Release(LobbyDetailsHandle);
+		return TDefaultErrorResultInternal<TSharedPtr<FLobbyDetailsEOS>>(MoveTemp(LobbyDetailsInfoResult.GetErrorValue()));
+	}
+
+	return TDefaultErrorResultInternal<TSharedPtr<FLobbyDetailsEOS>>(MakeShared<FLobbyDetailsEOS>(Prerequisites, LobbyDetailsInfoResult.GetOkValue().ToSharedRef(), ELobbyDetailsSource::UiEvent, LobbyDetailsHandle));
+}
+
+TDefaultErrorResultInternal<TSharedPtr<FLobbyDetailsEOS>> FLobbyDetailsEOS::CreateFromSearchResult(const TSharedRef<FLobbyPrerequisitesEOS>& Prerequisites, EOS_HLobbySearch SearchHandle, uint32_t ResultIndex)
+{
+	EOS_HLobbyDetails LobbyDetailsHandle = {};
+
+	EOS_LobbySearch_CopySearchResultByIndexOptions Options;
+	Options.ApiVersion = EOS_LOBBYSEARCH_COPYSEARCHRESULTBYINDEX_API_LATEST;
+	Options.LobbyIndex = ResultIndex;
+
+	EOS_EResult EOSResult = EOS_LobbySearch_CopySearchResultByIndex(SearchHandle, &Options, &LobbyDetailsHandle);
+	if (EOSResult != EOS_EResult::EOS_Success)
+	{
+		return TDefaultErrorResultInternal<TSharedPtr<FLobbyDetailsEOS>>(FromEOSError(EOSResult));
+	}
+
+	TDefaultErrorResultInternal<TSharedPtr<FLobbyDetailsInfoEOS>> LobbyDetailsInfoResult = FLobbyDetailsInfoEOS::Create(LobbyDetailsHandle);
+	if (LobbyDetailsInfoResult.IsError())
+	{
+		EOS_LobbyDetails_Release(LobbyDetailsHandle);
+		return TDefaultErrorResultInternal<TSharedPtr<FLobbyDetailsEOS>>(MoveTemp(LobbyDetailsInfoResult.GetErrorValue()));
+	}
+
+	return TDefaultErrorResultInternal<TSharedPtr<FLobbyDetailsEOS>>(MakeShared<FLobbyDetailsEOS>(Prerequisites, LobbyDetailsInfoResult.GetOkValue().ToSharedRef(), ELobbyDetailsSource::Search, LobbyDetailsHandle));
+}
+
+FLobbyDetailsEOS::~FLobbyDetailsEOS()
+{
+	EOS_LobbyDetails_Release(LobbyDetailsHandle);
+}
+
+TFuture<TDefaultErrorResultInternal<TSharedPtr<FLobby>>> FLobbyDetailsEOS::GetLobbyData(FOnlineAccountIdHandle LocalUserId, FOnlineLobbyIdHandle LobbyIdHandle) const
+{
+	EOS_LobbyDetails_GetMemberCountOptions GetMemberCountOptions = {};
+	GetMemberCountOptions.ApiVersion = EOS_LOBBYDETAILS_GETMEMBERCOUNT_API_LATEST;
+	const uint32_t MemberCount = EOS_LobbyDetails_GetMemberCount(LobbyDetailsHandle, &GetMemberCountOptions);
+
+	TSharedRef<TArray<EOS_ProductUserId>> MemberProductUserIds = MakeShared<TArray<EOS_ProductUserId>>();
+	MemberProductUserIds->Reserve(MemberCount);
+
+	for (uint32_t MemberIndex = 0; MemberIndex < MemberCount; ++MemberIndex)
+	{
+		EOS_LobbyDetails_GetMemberByIndexOptions GetMemberByIndexOptions = {};
+		GetMemberByIndexOptions.ApiVersion = EOS_LOBBYDETAILS_GETMEMBERBYINDEX_API_LATEST;
+		GetMemberByIndexOptions.MemberIndex = MemberIndex;
+		MemberProductUserIds->Emplace(EOS_LobbyDetails_GetMemberByIndex(LobbyDetailsHandle, &GetMemberByIndexOptions));
+	}
+
+	TPromise<TDefaultErrorResultInternal<TSharedPtr<FLobby>>> Promise;
+	TFuture<TDefaultErrorResultInternal<TSharedPtr<FLobby>>> Future = Promise.GetFuture();
+
+	Prerequisites->AuthInterface->ResolveAccountIds(LocalUserId, *MemberProductUserIds)
+	.Then([StrongThis = AsShared(), Promise = MoveTemp(Promise), MemberProductUserIds, LobbyIdHandle](TFuture<TArray<FOnlineAccountIdHandle>>&& Future) mutable
+	{
+		const TArray<FOnlineAccountIdHandle>& ResolvedAccountIds = Future.Get();
+		if (MemberProductUserIds->Num() != ResolvedAccountIds.Num())
+		{
+			// Todo: Errors
+			Promise.EmplaceValue(Errors::Unknown());
+			return;
+		}
+
+		EOS_LobbyDetails_GetLobbyOwnerOptions GetLobbyOwnerOptions = {};
+		GetLobbyOwnerOptions.ApiVersion = EOS_LOBBYDETAILS_GETLOBBYOWNER_API_LATEST;
+		const EOS_ProductUserId LobbyOwner = EOS_LobbyDetails_GetLobbyOwner(StrongThis->LobbyDetailsHandle, &GetLobbyOwnerOptions);
+
+		TDefaultErrorResultInternal<TMap<FLobbyAttributeId, FLobbyVariant>> GetAttributesResult = StrongThis->GetLobbyAttributes();
+		if (GetAttributesResult.IsError())
+		{
+			// Todo: Errors
+			Promise.EmplaceValue(Errors::Unknown(MoveTemp(GetAttributesResult.GetErrorValue())));
+			return;
+		}
+
+		TSharedPtr<FLobby> Lobby = MakeShared<FLobby>();
+		Lobby->LobbyId = LobbyIdHandle;
+		Lobby->Attributes = GetAttributesResult.GetOkValue();
+
+		for (int32 MemberIndex = 0; MemberIndex < MemberProductUserIds->Num(); ++MemberIndex)
+		{
+			const EOS_ProductUserId MemberProductUserId = (*MemberProductUserIds)[MemberIndex];
+			const FOnlineAccountIdHandle ResolvedProductUserId = ResolvedAccountIds[MemberIndex];
+
+			if (MemberProductUserId == LobbyOwner)
+			{
+				Lobby->OwnerAccountId = ResolvedProductUserId;
+			}
+
+			TDefaultErrorResultInternal<TMap<FLobbyAttributeId, FLobbyVariant>> GetMemberAttributesResult = StrongThis->GetLobbyMemberAttributes(MemberProductUserId);
+			if (GetMemberAttributesResult.IsError())
+			{
+				// Todo: Errors
+				Promise.EmplaceValue(Errors::Unknown(MoveTemp(GetMemberAttributesResult.GetErrorValue())));
+				return;
+			}
+
+			Lobby->Members.Emplace(StrongThis->CreateLobbyMember(ResolvedProductUserId, MoveTemp(GetMemberAttributesResult.GetOkValue())));
+		}
+
+		Promise.EmplaceValue(MoveTemp(Lobby));
+	});
+
+	return Future;
+}
+
+TDefaultErrorResultInternal<TSharedPtr<FLobbyMember>> FLobbyDetailsEOS::GetLobbyMemberData(FOnlineAccountIdHandle MemberHandle) const
+{
+	TDefaultErrorResultInternal<TMap<FLobbyAttributeId, FLobbyVariant>> MemberAttributesResult = GetLobbyMemberAttributes(GetProductUserIdChecked(MemberHandle));
+	if (MemberAttributesResult.IsError())
+	{
+		return TDefaultErrorResultInternal<TSharedPtr<FLobbyMember>>(MoveTemp(MemberAttributesResult.GetErrorValue()));
+	}
+	else
+	{
+		return TDefaultErrorResultInternal<TSharedPtr<FLobbyMember>>(CreateLobbyMember(MemberHandle, MoveTemp(MemberAttributesResult.GetOkValue())));
+	}
+}
+
+TFuture<EOS_EResult> FLobbyDetailsEOS::ApplyLobbyDataUpdates(FOnlineAccountIdHandle LocalUserId, FLobbyDataChanges Changes) const
+{
+	EOS_HLobbyModification LobbyModificationHandle = {};
+
+	ON_SCOPE_EXIT
+	{
+		EOS_LobbyModification_Release(LobbyModificationHandle);
+	};
+
+	// Create lobby modification handle.
+	EOS_Lobby_UpdateLobbyModificationOptions ModificationOptions = {};
+	ModificationOptions.ApiVersion = EOS_LOBBY_UPDATELOBBYMODIFICATION_API_LATEST;
+	ModificationOptions.LocalUserId = GetProductUserIdChecked(LocalUserId);
+	ModificationOptions.LobbyId = GetInfo()->GetLobbyIdEOS();
+
+	EOS_EResult EOSResultCode = EOS_Lobby_UpdateLobbyModification(Prerequisites->LobbyInterfaceHandle, &ModificationOptions, &LobbyModificationHandle);
+	if (EOSResultCode != EOS_EResult::EOS_Success)
+	{
+		// Todo: Errors
+		return MakeFulfilledPromise<EOS_EResult>(EOSResultCode).GetFuture();
+	}
+
+	if (Changes.JoinPolicy)
+	{
+		// Set lobby join policy.
+		EOS_LobbyModification_SetPermissionLevelOptions SetPermissionOptions = {};
+		SetPermissionOptions.ApiVersion = EOS_LOBBYMODIFICATION_SETPERMISSIONLEVEL_API_LATEST;
+		SetPermissionOptions.PermissionLevel = TranslateJoinPolicy(*Changes.JoinPolicy);
+
+		EOSResultCode = EOS_LobbyModification_SetPermissionLevel(LobbyModificationHandle, &SetPermissionOptions);
+		if (EOSResultCode != EOS_EResult::EOS_Success)
+		{
+			// Todo: Errors
+			return MakeFulfilledPromise<EOS_EResult>(EOSResultCode).GetFuture();
+		}
+	}
+
+	// Add attributes.
+	for (const TPair<FLobbyAttributeId, FLobbyVariant>& MutatedAttribute : Changes.MutatedAttributes)
+	{
+		const FTCHARToUTF8 KeyConverter(*MutatedAttribute.Key.ToString());
+		const FTCHARToUTF8 ValueConverter(*MutatedAttribute.Value);
+
+		// Todo: handle variant properly.
+		// Todo: Schema things.
+
+		EOS_Lobby_AttributeData AttributeData;
+		AttributeData.Key = KeyConverter.Get();
+		AttributeData.Value.AsUtf8 = ValueConverter.Get();
+		AttributeData.ValueType = EOS_ELobbyAttributeType::EOS_AT_STRING;
+
+		EOS_LobbyModification_AddAttributeOptions AddAttributeOptions = {};
+		AddAttributeOptions.ApiVersion = EOS_LOBBYMODIFICATION_ADDATTRIBUTE_API_LATEST;
+		AddAttributeOptions.Attribute = &AttributeData;
+		//AddAttributeOptions.Visibility; // todo - get from schema
+		EOSResultCode = EOS_LobbyModification_AddAttribute(LobbyModificationHandle, &AddAttributeOptions);
+		if (EOSResultCode != EOS_EResult::EOS_Success)
+		{
+			// Todo: Errors
+			return MakeFulfilledPromise<EOS_EResult>(EOSResultCode).GetFuture();
+		}
+	}
+
+	// Remove attributes.
+	for (const FLobbyAttributeId& ClearedAttribute : Changes.ClearedAttributes)
+	{
+		const FTCHARToUTF8 KeyConverter(*ClearedAttribute.ToString());
+
+		EOS_LobbyModification_RemoveAttributeOptions RemoveAttributeOptions = {};
+		RemoveAttributeOptions.ApiVersion = EOS_LOBBYMODIFICATION_REMOVEATTRIBUTE_API_LATEST;
+		RemoveAttributeOptions.Key = KeyConverter.Get();
+		EOSResultCode = EOS_LobbyModification_RemoveAttribute(LobbyModificationHandle, &RemoveAttributeOptions);
+		if (EOSResultCode != EOS_EResult::EOS_Success)
+		{
+			// Todo: Errors
+			return MakeFulfilledPromise<EOS_EResult>(EOSResultCode).GetFuture();
+		}
+	}
+
+	// Apply lobby updates.
+	EOS_Lobby_UpdateLobbyOptions UpdateLobbyOptions = {};
+	UpdateLobbyOptions.ApiVersion = EOS_LOBBY_UPDATELOBBY_API_LATEST;
+	UpdateLobbyOptions.LobbyModificationHandle = LobbyModificationHandle;
+
+	TPromise<EOS_EResult> Promise;
+	TFuture<EOS_EResult> Future = Promise.GetFuture();
+
+	EOS_Async<EOS_Lobby_UpdateLobbyCallbackInfo>(EOS_Lobby_UpdateLobby, Prerequisites->LobbyInterfaceHandle, UpdateLobbyOptions)
+	.Then([Promise = MoveTemp(Promise)](TFuture<const EOS_Lobby_UpdateLobbyCallbackInfo*> Future) mutable
+	{
+		Promise.EmplaceValue(Future.Get()->ResultCode);
+	});
+
+	return Future;
+}
+
+TFuture<EOS_EResult> FLobbyDetailsEOS::ApplyLobbyMemberDataUpdates(FOnlineAccountIdHandle LocalUserId, FLobbyMemberDataChanges Changes) const
+{
+	EOS_HLobbyModification LobbyModificationHandle = {};
+
+	ON_SCOPE_EXIT
+	{
+		EOS_LobbyModification_Release(LobbyModificationHandle);
+	};
+
+	// Create lobby modification handle.
+	EOS_Lobby_UpdateLobbyModificationOptions ModificationOptions = {};
+	ModificationOptions.ApiVersion = EOS_LOBBY_UPDATELOBBYMODIFICATION_API_LATEST;
+	ModificationOptions.LocalUserId = GetProductUserIdChecked(LocalUserId);
+	ModificationOptions.LobbyId = GetInfo()->GetLobbyIdEOS();
+
+	EOS_EResult EOSResultCode = EOS_Lobby_UpdateLobbyModification(Prerequisites->LobbyInterfaceHandle, &ModificationOptions, &LobbyModificationHandle);
+	if (EOSResultCode != EOS_EResult::EOS_Success)
+	{
+		// Todo: Errors
+		return MakeFulfilledPromise<EOS_EResult>(EOSResultCode).GetFuture();
+	}
+
+	// Add member attributes.
+	for (const TPair<FLobbyAttributeId, FLobbyVariant>& MutatedAttribute : Changes.MutatedAttributes)
+	{
+		const FTCHARToUTF8 KeyConverter(*MutatedAttribute.Key.ToString());
+		const FTCHARToUTF8 ValueConverter(*MutatedAttribute.Value);
+
+		// Todo: handle variant properly.
+		// Todo: Schema things.
+
+		EOS_Lobby_AttributeData AttributeData;
+		AttributeData.Key = KeyConverter.Get();
+		AttributeData.Value.AsUtf8 = ValueConverter.Get();
+
+		EOS_LobbyModification_AddMemberAttributeOptions AddMemberAttributeOptions = {};
+		AddMemberAttributeOptions.ApiVersion = EOS_LOBBYMODIFICATION_ADDMEMBERATTRIBUTE_API_LATEST;
+		AddMemberAttributeOptions.Attribute = &AttributeData;
+		//AddMemberAttributeOptions.Visibility; // todo - get from schema
+		EOSResultCode = EOS_LobbyModification_AddMemberAttribute(LobbyModificationHandle, &AddMemberAttributeOptions);
+		if (EOSResultCode != EOS_EResult::EOS_Success)
+		{
+			// Todo: Errors
+			return MakeFulfilledPromise<EOS_EResult>(EOSResultCode).GetFuture();
+		}
+	}
+
+	// Remove member attributes.
+	for (const FLobbyAttributeId& ClearedAttribute : Changes.ClearedAttributes)
+	{
+		const FTCHARToUTF8 KeyConverter(*ClearedAttribute.ToString());
+
+		EOS_LobbyModification_RemoveMemberAttributeOptions RemoveMemberAttributeOptions = {};
+		RemoveMemberAttributeOptions.ApiVersion = EOS_LOBBYMODIFICATION_REMOVEMEMBERATTRIBUTE_API_LATEST;
+		RemoveMemberAttributeOptions.Key = KeyConverter.Get();
+		EOSResultCode = EOS_LobbyModification_RemoveMemberAttribute(LobbyModificationHandle, &RemoveMemberAttributeOptions);
+		if (EOSResultCode != EOS_EResult::EOS_Success)
+		{
+			// Todo: Errors
+			return MakeFulfilledPromise<EOS_EResult>(EOSResultCode).GetFuture();
+		}
+	}
+
+	// Apply lobby updates.
+	EOS_Lobby_UpdateLobbyOptions UpdateLobbyOptions = {};
+	UpdateLobbyOptions.ApiVersion = EOS_LOBBY_UPDATELOBBY_API_LATEST;
+	UpdateLobbyOptions.LobbyModificationHandle = LobbyModificationHandle;
+
+	TPromise<EOS_EResult> Promise;
+	TFuture<EOS_EResult> Future = Promise.GetFuture();
+
+	EOS_Async<EOS_Lobby_UpdateLobbyCallbackInfo>(EOS_Lobby_UpdateLobby, Prerequisites->LobbyInterfaceHandle, UpdateLobbyOptions)
+	.Then([Promise = MoveTemp(Promise)](TFuture<const EOS_Lobby_UpdateLobbyCallbackInfo*> Future) mutable
+	{
+		Promise.EmplaceValue(Future.Get()->ResultCode);
+	});
+
+	return Future;
+}
+
+FLobbyDetailsEOS::FLobbyDetailsEOS(
+	const TSharedRef<FLobbyPrerequisitesEOS>& Prerequisites,
+	const TSharedRef<FLobbyDetailsInfoEOS>& LobbyDetailsInfo,
+	ELobbyDetailsSource LobbyDetailsSource,
+	EOS_HLobbyDetails LobbyDetailsHandle)
+	: Prerequisites(Prerequisites)
+	, LobbyDetailsInfo(LobbyDetailsInfo)
+	, LobbyDetailsHandle(LobbyDetailsHandle)
+	, LobbyDetailsSource(LobbyDetailsSource)
+{
+}
+
+TPair<FLobbyAttributeId, FLobbyVariant> FLobbyDetailsEOS::TranslateLobbyAttribute(const EOS_Lobby_Attribute& LobbyAttribute) const
+{
+	const FLobbyAttributeId AttributeId = detail::TranslateLobbyAttributeId(LobbyAttribute.Data->Key);
+	FLobbyVariant VariantData;
+
+	switch (LobbyAttribute.Data->ValueType)
+	{
+	case EOS_ELobbyAttributeType::EOS_AT_BOOLEAN:
+		// Todo:
+		break;
+
+	case EOS_ELobbyAttributeType::EOS_AT_INT64:
+		// Todo:
+		break;
+
+	case EOS_ELobbyAttributeType::EOS_AT_DOUBLE:
+		// Todo:
+		break;
+
+	case EOS_ELobbyAttributeType::EOS_AT_STRING:
+		VariantData = UTF8_TO_TCHAR(LobbyAttribute.Data->Value.AsUtf8);
+		break;
+
+	default:
+		// Todo: log
+		break;
+	}
+
+	return TPair<FLobbyAttributeId, FLobbyVariant>(AttributeId, MoveTemp(VariantData));
+}
+
+TDefaultErrorResultInternal<TMap<FLobbyAttributeId, FLobbyVariant>> FLobbyDetailsEOS::GetLobbyAttributes() const
+{
+	TMap<FLobbyAttributeId, FLobbyVariant> Attributes;
+
+	EOS_LobbyDetails_GetAttributeCountOptions GetAttributeCountOptions = {};
+	GetAttributeCountOptions.ApiVersion = EOS_LOBBYDETAILS_GETATTRIBUTECOUNT_API_LATEST;
+
+	const uint32_t AttributeCount = EOS_LobbyDetails_GetAttributeCount(LobbyDetailsHandle, &GetAttributeCountOptions);
+	for (uint32_t AttributeIndex = 0; AttributeIndex < AttributeCount; ++AttributeIndex)
+	{
+		EOS_LobbyDetails_CopyAttributeByIndexOptions CopyAttributeByIndexOptions = {};
+		CopyAttributeByIndexOptions.ApiVersion = EOS_LOBBYDETAILS_COPYATTRIBUTEBYINDEX_API_LATEST;
+		CopyAttributeByIndexOptions.AttrIndex = AttributeIndex;
+
+		EOS_Lobby_Attribute* LobbyAttribute = nullptr;
+		ON_SCOPE_EXIT
+		{
+			EOS_Lobby_Attribute_Release(LobbyAttribute);
+		};
+
+		EOS_EResult EOSResult = EOS_LobbyDetails_CopyAttributeByIndex(LobbyDetailsHandle, &CopyAttributeByIndexOptions, &LobbyAttribute);
+		if (EOSResult != EOS_EResult::EOS_Success)
+		{
+			return TDefaultErrorResultInternal<TMap<FLobbyAttributeId, FLobbyVariant>>(FromEOSError(EOSResult));
+		}
+
+		// Todo: Schema things.
+		TPair<FLobbyAttributeId, FLobbyVariant> TranslatedAttribute = TranslateLobbyAttribute(*LobbyAttribute);
+		Attributes.Emplace(TranslatedAttribute.Key, MoveTemp(TranslatedAttribute.Value));
+	}
+
+	return TDefaultErrorResultInternal<TMap<FLobbyAttributeId, FLobbyVariant>>(MoveTemp(Attributes));
+}
+
+TDefaultErrorResultInternal<TMap<FLobbyAttributeId, FLobbyVariant>> FLobbyDetailsEOS::GetLobbyMemberAttributes(EOS_ProductUserId TargetMemberProductUserId) const
+{
+	TMap<FLobbyAttributeId, FLobbyVariant> Attributes;
+
+	EOS_LobbyDetails_GetMemberAttributeCountOptions GetMemberAttributeCountOptions = {};
+	GetMemberAttributeCountOptions.ApiVersion = EOS_LOBBYDETAILS_GETMEMBERATTRIBUTECOUNT_API_LATEST;
+	GetMemberAttributeCountOptions.TargetUserId = TargetMemberProductUserId;
+
+	const uint32_t MemberAttributeCount = EOS_LobbyDetails_GetMemberAttributeCount(LobbyDetailsHandle, &GetMemberAttributeCountOptions);
+	for (uint32_t MemberAttributeIndex = 0; MemberAttributeIndex < MemberAttributeCount; ++MemberAttributeIndex)
+	{
+		EOS_LobbyDetails_CopyMemberAttributeByIndexOptions CopyMemberAttributeByIndexOptions = {};
+		CopyMemberAttributeByIndexOptions.ApiVersion = EOS_LOBBYDETAILS_COPYMEMBERATTRIBUTEBYINDEX_API_LATEST;
+		CopyMemberAttributeByIndexOptions.TargetUserId = TargetMemberProductUserId;
+		CopyMemberAttributeByIndexOptions.AttrIndex = MemberAttributeIndex;
+
+		EOS_Lobby_Attribute* LobbyAttribute = nullptr;
+		ON_SCOPE_EXIT
+		{
+			EOS_Lobby_Attribute_Release(LobbyAttribute);
+		};
+
+		EOS_EResult EOSResult = EOS_LobbyDetails_CopyMemberAttributeByIndex(LobbyDetailsHandle, &CopyMemberAttributeByIndexOptions, &LobbyAttribute);
+		if (EOSResult != EOS_EResult::EOS_Success)
+		{
+			return TDefaultErrorResultInternal<TMap<FLobbyAttributeId, FLobbyVariant>>(FromEOSError(EOSResult));
+		}
+
+		// Todo: Schema things.
+		TPair<FLobbyAttributeId, FLobbyVariant> TranslatedAttribute = TranslateLobbyAttribute(*LobbyAttribute);
+		Attributes.Emplace(TranslatedAttribute.Key, MoveTemp(TranslatedAttribute.Value));
+	}
+
+	return TDefaultErrorResultInternal<TMap<FLobbyAttributeId, FLobbyVariant>>(MoveTemp(Attributes));
+}
+
+TSharedRef<FLobbyMember> FLobbyDetailsEOS::CreateLobbyMember(FOnlineAccountIdHandle MemberHandle, TMap<FLobbyAttributeId, FLobbyVariant> Attributes) const
+{
+	TSharedRef<FLobbyMember> LobbyMember = MakeShared<FLobbyMember>();
+	LobbyMember->AccountId = MemberHandle;
+	LobbyMember->Attributes = MoveTemp(Attributes);
+
+	//Todo: 
+	//LobbyMember->PlatformAccountId;
+	//LobbyMember->PlatformDisplayName;
+
+	return LobbyMember;
+}
+
+TDefaultErrorResultInternal<TSharedPtr<FLobbyDetailsInfoEOS>> FLobbyDetailsInfoEOS::Create(EOS_HLobbyDetails LobbyDetailsHandle)
+{
+	EOS_LobbyDetails_CopyInfoOptions CopyInfoOptions = {};
+	CopyInfoOptions.ApiVersion = EOS_LOBBYDETAILS_COPYINFO_API_LATEST;
+
+	EOS_LobbyDetails_Info* LobbyDetailsInfo = nullptr;
+	EOS_EResult EOSResult = EOS_LobbyDetails_CopyInfo(LobbyDetailsHandle, &CopyInfoOptions, &LobbyDetailsInfo);
+	if (EOSResult != EOS_EResult::EOS_Success)
+	{
+		return TDefaultErrorResultInternal<TSharedPtr<FLobbyDetailsInfoEOS>>(FromEOSError(EOSResult));
+	}
+
+	TSharedRef<FLobbyDetailsInfoEOS> Result = MakeShared<FLobbyDetailsInfoEOS>();
+	Result->LobbyDetailsInfo.Reset(LobbyDetailsInfo);
+
+	// Resolve lobby info.
+
+	return TDefaultErrorResultInternal<TSharedPtr<FLobbyDetailsInfoEOS>>(MoveTemp(Result));
+}
+
+FLobbyDataEOS::~FLobbyDataEOS()
+{
+	if (UnregisterFn)
+	{
+		UnregisterFn(LobbyIdHandle);
+	}
+}
+
+TSharedRef<FLobbyNotificationPauseHandle> FLobbyDataEOS::PauseLobbyNotifications(FOnlineAccountIdHandle LocalUserId)
+{
+	if (TSharedRef<FLobbyNotificationPauseHandle>* ExistingPause = NotificationPauses.Find(LocalUserId))
+	{
+		return *ExistingPause;
+	}
+	else
+	{
+		TSharedRef<FLobbyNotificationPauseHandle> NewPause = MakeShared<detail::FLobbyNotificationPauseHandleImpl>(
+		[WeakThis = AsWeak(), LocalUserId]()
+		{
+			if (TSharedPtr<FLobbyDataEOS> StrongThis = WeakThis.Pin())
+			{
+				StrongThis->NotificationPauses.Remove(LocalUserId);
+			}
+		});
+
+		NotificationPauses.Add(LocalUserId, NewPause);
+		return NewPause;
+	}
+}
+
+bool FLobbyDataEOS::CanNotify(FOnlineAccountIdHandle LocalUserId)
+{
+	return !LocalMembers.IsEmpty() && NotificationPauses.Find(LocalUserId) == nullptr;
+}
+
+void FLobbyDataEOS::SetUserLobbyDetails(FOnlineAccountIdHandle LocalUserId, const TSharedPtr<FLobbyDetailsEOS>& LobbyDetails)
+{
+	if (TSharedPtr<FLobbyDetailsEOS> ExistingDetails = GetUserLobbyDetails(LocalUserId))
+	{
+		if (ExistingDetails->GetDetailsSource() < LobbyDetails->GetDetailsSource())
+		{
+			return;
+		}
+	}
+
+	UserLobbyDetails.Add(LocalUserId, LobbyDetails);
+}
+
+TSharedPtr<FLobbyDetailsEOS> FLobbyDataEOS::GetUserLobbyDetails(FOnlineAccountIdHandle LocalUserId) const
+{
+	const TSharedPtr<FLobbyDetailsEOS>* Result = UserLobbyDetails.Find(LocalUserId);
+	return Result ? *Result : TSharedPtr<FLobbyDetailsEOS>();
+}
+
+FLobbyDataEOS::FLobbyDataEOS(
+	FOnlineLobbyIdHandle LobbyIdHandle,
+	const TSharedRef<FLobby>& LobbyImpl,
+	const TSharedRef<FLobbyDetailsInfoEOS>& LobbyDetailsInfo,
+	FUnregisterFn UnregisterFn)
+	: LobbyIdHandle(LobbyIdHandle)
+	, LobbyImpl(LobbyImpl)
+	, LobbyDetailsInfo(LobbyDetailsInfo)
+	, UnregisterFn(MoveTemp(UnregisterFn))
+	, LobbyId(detail::TranslateLobbyId(LobbyDetailsInfo->GetLobbyIdEOS()))
+{
+}
+
+TFuture<TDefaultErrorResultInternal<TSharedPtr<FLobbyDataEOS>>> FLobbyDataEOS::Create(
+	FOnlineAccountIdHandle LocalUserId,
+	FOnlineLobbyIdHandle LobbyIdHandle,
+	const TSharedRef<FLobbyDetailsEOS>& LobbyDetails,
+	FUnregisterFn UnregisterFn)
+{
+	TPromise<TDefaultErrorResultInternal<TSharedPtr<FLobbyDataEOS>>> Promise;
+	TFuture<TDefaultErrorResultInternal<TSharedPtr<FLobbyDataEOS>>> Future = Promise.GetFuture();
+	
+	LobbyDetails->GetLobbyData(LocalUserId, LobbyIdHandle)
+	.Then(
+	[
+		Promise = MoveTemp(Promise),
+		LocalUserId,
+		LobbyIdHandle,
+		LobbyDetails,
+		UnregisterFn = MoveTemp(UnregisterFn)
+	]
+	(TFuture<TDefaultErrorResultInternal<TSharedPtr<FLobby>>>&& Future) mutable
+	{
+		if (Future.Get().IsError())
+		{
+			Promise.EmplaceValue(MoveTemp(Future.Get().GetErrorValue()));
+		}
+		else
+		{
+			Promise.EmplaceValue(MakeShared<FLobbyDataEOS>(
+				LobbyIdHandle,
+				Future.Get().GetOkValue().ToSharedRef(),
+				LobbyDetails->GetInfo(),
+				MoveTemp(UnregisterFn)));
+		}
+	});
+
+	return Future;
+}
+
+FLobbyDataRegistryEOS::FLobbyDataRegistryEOS(const TSharedRef<FLobbyPrerequisitesEOS>& Prerequisites)
+	: Prerequisites(Prerequisites)
+{
+}
+
+TSharedPtr<FLobbyDataEOS> FLobbyDataRegistryEOS::Find(EOS_LobbyId EOSLobbyId) const
+{
+	const TWeakPtr<FLobbyDataEOS>* Result = LobbyIdIndex.Find(detail::TranslateLobbyId(EOSLobbyId));
+	return Result ? Result->Pin() : TSharedPtr<FLobbyDataEOS>();
+}
+
+TSharedPtr<FLobbyDataEOS> FLobbyDataRegistryEOS::Find(FOnlineLobbyIdHandle LobbyIdHandle) const
+{
+	const TWeakPtr<FLobbyDataEOS>* Result = LobbyIdHandleIndex.Find(LobbyIdHandle);
+	return Result ? Result->Pin() : TSharedPtr<FLobbyDataEOS>();
+}
+
+TFuture<TDefaultErrorResultInternal<TSharedPtr<FLobbyDataEOS>>> FLobbyDataRegistryEOS::FindOrCreateFromLobbyId(FOnlineAccountIdHandle LocalUserId, EOS_LobbyId EOSLobbyId)
+{
+	TDefaultErrorResultInternal<TSharedPtr<FLobbyDetailsEOS>> LobbyDetailsResult = FLobbyDetailsEOS::CreateFromLobbyId(Prerequisites, EOSLobbyId, LocalUserId);
+	if (LobbyDetailsResult.IsError())
+	{
+		return MakeFulfilledPromise<TDefaultErrorResultInternal<TSharedPtr<FLobbyDataEOS>>>(MoveTemp(LobbyDetailsResult.GetErrorValue())).GetFuture();
+	}
+
+	return FindOrCreateFromLobbyDetails(LocalUserId, LobbyDetailsResult.GetOkValue().ToSharedRef());
+}
+
+TFuture<TDefaultErrorResultInternal<TSharedPtr<FLobbyDataEOS>>> FLobbyDataRegistryEOS::FindOrCreateFromUiEventId(FOnlineAccountIdHandle LocalUserId, EOS_UI_EventId UiEventId)
+{
+	TDefaultErrorResultInternal<TSharedPtr<FLobbyDetailsEOS>> LobbyDetailsResult = FLobbyDetailsEOS::CreateFromUiEventId(Prerequisites, UiEventId);
+	if (LobbyDetailsResult.IsError())
+	{
+		return MakeFulfilledPromise<TDefaultErrorResultInternal<TSharedPtr<FLobbyDataEOS>>>(MoveTemp(LobbyDetailsResult.GetErrorValue())).GetFuture();
+	}
+
+	return FindOrCreateFromLobbyDetails(LocalUserId, LobbyDetailsResult.GetOkValue().ToSharedRef());
+}
+
+TFuture<TDefaultErrorResultInternal<TSharedPtr<FLobbyDataEOS>>> FLobbyDataRegistryEOS::FindOrCreateFromLobbyDetails(FOnlineAccountIdHandle LocalUserId, const TSharedRef<FLobbyDetailsEOS>& LobbyDetails)
+{
+	if (TSharedPtr<FLobbyDataEOS> FindResult = Find(LobbyDetails->GetInfo()->GetLobbyIdEOS()))
+	{
+		FindResult->SetUserLobbyDetails(LocalUserId, LobbyDetails);
+		return MakeFulfilledPromise<TDefaultErrorResultInternal<TSharedPtr<FLobbyDataEOS>>>(FindResult).GetFuture();
+	}
+
+	TPromise<TDefaultErrorResultInternal<TSharedPtr<FLobbyDataEOS>>> Promise;
+	TFuture<TDefaultErrorResultInternal<TSharedPtr<FLobbyDataEOS>>> Future = Promise.GetFuture();
+
+	const FOnlineLobbyIdHandle LobbyId = FOnlineLobbyIdHandle(EOnlineServices::Epic, NextHandleIndex++);
+	FLobbyDataEOS::Create(LocalUserId, LobbyId, LobbyDetails, MakeUnregisterFn())
+	.Then([WeakThis = AsWeak(), Promise = MoveTemp(Promise), LocalUserId, LobbyDetails](TFuture<TDefaultErrorResultInternal<TSharedPtr<FLobbyDataEOS>>>&& Future) mutable
+	{
+		if (TSharedPtr<FLobbyDataRegistryEOS> StrongThis = WeakThis.Pin())
+		{
+			if (Future.Get().IsOk())
+			{
+				StrongThis->Register(Future.Get().GetOkValue().ToSharedRef());
+			}
+		}
+
+		Future.Get().GetOkValue()->SetUserLobbyDetails(LocalUserId, LobbyDetails);
+		Promise.EmplaceValue(MoveTempIfPossible(Future.Get()));
+	});
+
+	return Future;
+}
+
+void FLobbyDataRegistryEOS::Register(const TSharedRef<FLobbyDataEOS>& LobbyIdHandleData)
+{
+	LobbyIdIndex.Add(LobbyIdHandleData->GetLobbyId(), LobbyIdHandleData);
+	LobbyIdHandleIndex.Add(LobbyIdHandleData->GetLobbyIdHandle(), LobbyIdHandleData);
+}
+
+void FLobbyDataRegistryEOS::Unregister(FOnlineLobbyIdHandle LobbyIdHandle)
+{
+	if (TSharedPtr<FLobbyDataEOS> HandleData = Find(LobbyIdHandle))
+	{
+		LobbyIdIndex.Remove(HandleData->GetLobbyId());
+		LobbyIdHandleIndex.Remove(HandleData->GetLobbyIdHandle());
+	}
+}
+
+FLobbyDataEOS::FUnregisterFn FLobbyDataRegistryEOS::MakeUnregisterFn()
+{
+	return [WeakThis = AsWeak()](FOnlineLobbyIdHandle LobbyId)
+	{
+		if (TSharedPtr<FLobbyDataRegistryEOS> StrongThis = WeakThis.Pin())
+		{
+			StrongThis->Unregister(LobbyId);
+		}
+	};
+}
+
+TFuture<TDefaultErrorResultInternal<TSharedPtr<FLobbyInviteDataEOS>>> FLobbyInviteDataEOS::CreateFromInviteId(
+	const TSharedRef<FLobbyPrerequisitesEOS>& Prerequisites,
+	const TSharedRef<FLobbyDataRegistryEOS>& LobbyDataRegistry,
+	FOnlineAccountIdHandle LocalUserId,
+	const char* InviteIdEOS,
+	EOS_ProductUserId Sender)
+{
+	TDefaultErrorResultInternal<TSharedPtr<FLobbyDetailsEOS>> LobbyDetailsResult = FLobbyDetailsEOS::CreateFromInviteId(Prerequisites, InviteIdEOS);
+	if (LobbyDetailsResult.IsError())
+	{
+		return MakeFulfilledPromise<TDefaultErrorResultInternal<TSharedPtr<FLobbyInviteDataEOS>>>(MoveTemp(LobbyDetailsResult.GetErrorValue())).GetFuture();
+	}
+
+	TSharedRef<FLobbyDetailsEOS> LobbyDetails = LobbyDetailsResult.GetOkValue().ToSharedRef();
+	TPromise<TDefaultErrorResultInternal<TSharedPtr<FLobbyInviteDataEOS>>> Promise;
+	TFuture<TDefaultErrorResultInternal<TSharedPtr<FLobbyInviteDataEOS>>> Future = Promise.GetFuture();
+
+	// Search for existing lobby data so that the LobbyIdHandle will match.
+	TSharedRef<FLobbyInviteIdEOS> InviteId = MakeShared<FLobbyInviteIdEOS>(InviteIdEOS);
+	LobbyDataRegistry->FindOrCreateFromLobbyDetails(LocalUserId, LobbyDetails)
+	.Then([Promise = MoveTemp(Promise), InviteId, LocalUserId, Sender, LobbyDetails](TFuture<TDefaultErrorResultInternal<TSharedPtr<FLobbyDataEOS>>>&& Future) mutable
+	{
+		if (Future.Get().IsError())
+		{
+			Promise.EmplaceValue(MoveTemp(Future.Get().GetErrorValue()));
+			return;
+		}
+
+		// Once the lobby data has been resolved the FOnlineAccountIdHandle for the sender is expected to be in the AccountID cache.
+		const FOnlineAccountIdHandle SenderUserId = FindAccountId(Sender);
+		if (!SenderUserId.IsValid())
+		{
+			// Todo: Errors.
+			Promise.EmplaceValue(Errors::Unknown());
+			return;
+		}
+
+		Promise.EmplaceValue(MakeShared<FLobbyInviteDataEOS>(InviteId, LocalUserId, SenderUserId, LobbyDetails, Future.Get().GetOkValue().ToSharedRef()));
+	});
+
+	return Future;
+}
+
+FLobbyInviteDataEOS::FLobbyInviteDataEOS(
+	const TSharedRef<FLobbyInviteIdEOS>& InviteIdEOS,
+	FOnlineAccountIdHandle Receiver,
+	FOnlineAccountIdHandle Sender,
+	const TSharedRef<FLobbyDetailsEOS>& LobbyDetails,
+	const TSharedRef<FLobbyDataEOS>& LobbyData)
+	: InviteIdEOS(InviteIdEOS)
+	, Receiver(Receiver)
+	, Sender(Sender)
+	, LobbyDetails(LobbyDetails)
+	, LobbyData(LobbyData)
+	, InviteId(detail::TranslateLobbyInviteId(InviteIdEOS->Get()))
+{
+}
+
+TFuture<TDefaultErrorResultInternal<TSharedPtr<FLobbySearchEOS>>> FLobbySearchEOS::Create(
+	const TSharedRef<FLobbyPrerequisitesEOS>& Prerequisites,
+	const TSharedRef<FLobbyDataRegistryEOS>& LobbyRegistry,
+	const FLobbySearchParameters& Params)
+{
+	TSharedRef<FSearchHandle> SearchHandle = MakeShared<FSearchHandle>();
+
+	EOS_Lobby_CreateLobbySearchOptions CreateLobbySearchOptions = {};
+	CreateLobbySearchOptions.ApiVersion = EOS_LOBBY_CREATELOBBYSEARCH_API_LATEST;
+	CreateLobbySearchOptions.MaxResults = Params.MaxResults;
+
+	EOS_EResult EOSResult = EOS_Lobby_CreateLobbySearch(Prerequisites->LobbyInterfaceHandle, &CreateLobbySearchOptions, &SearchHandle->Get());
+	if (EOSResult != EOS_EResult::EOS_Success)
+	{
+		// todo: errors
+		return MakeFulfilledPromise<TDefaultErrorResultInternal<TSharedPtr<FLobbySearchEOS>>>(FromEOSError(EOSResult)).GetFuture();
+	}
+
+	if (Params.LobbyId)
+	{
+		TSharedPtr<FLobbyDataEOS> LobbyData = LobbyRegistry->Find(*Params.LobbyId);
+		if (!LobbyData)
+		{
+			return MakeFulfilledPromise<TDefaultErrorResultInternal<TSharedPtr<FLobbySearchEOS>>>(Errors::InvalidParams()).GetFuture();
+		}
+
+		EOS_LobbySearch_SetLobbyIdOptions SetLobbyIdOptions = {};
+		SetLobbyIdOptions.ApiVersion = EOS_LOBBYSEARCH_SETLOBBYID_API_LATEST;
+		SetLobbyIdOptions.LobbyId = LobbyData->GetLobbyIdEOS();
+
+		EOSResult = EOS_LobbySearch_SetLobbyId(SearchHandle->Get(), &SetLobbyIdOptions);
+		if (EOSResult != EOS_EResult::EOS_Success)
+		{
+			// todo: errors
+			return MakeFulfilledPromise<TDefaultErrorResultInternal<TSharedPtr<FLobbySearchEOS>>>(FromEOSError(EOSResult)).GetFuture();
+		}
+	}
+
+	if (Params.TargetUser)
+	{
+		EOS_LobbySearch_SetTargetUserIdOptions SetTargetUserIdOptions = {};
+		SetTargetUserIdOptions.ApiVersion = EOS_LOBBYSEARCH_SETTARGETUSERID_API_LATEST;
+		SetTargetUserIdOptions.TargetUserId = GetProductUserIdChecked(*Params.TargetUser);
+
+		EOSResult = EOS_LobbySearch_SetTargetUserId(SearchHandle->Get(), &SetTargetUserIdOptions);
+		if (EOSResult != EOS_EResult::EOS_Success)
+		{
+			// todo: errors
+			return MakeFulfilledPromise<TDefaultErrorResultInternal<TSharedPtr<FLobbySearchEOS>>>(FromEOSError(EOSResult)).GetFuture();
+		}
+	}
+
+	for (const FFindLobbySearchFilter& Filter :  Params.Filters)
+	{
+		FTCHARToUTF8 Key(*Filter.AttributeName.ToString());
+		FTCHARToUTF8 Value(*Filter.ComparisonValue);
+
+		EOS_Lobby_AttributeData AttributeData;
+		AttributeData.ApiVersion = EOS_LOBBY_ATTRIBUTEDATA_API_LATEST;
+
+		// todo: support other attribute types.
+		AttributeData.Key = Key.Get();
+		AttributeData.Value.AsUtf8 = Value.Get();
+		AttributeData.ValueType = EOS_ELobbyAttributeType::EOS_AT_STRING;
+
+		EOS_LobbySearch_SetParameterOptions SetParameterOptions = {};
+		SetParameterOptions.ApiVersion = EOS_LOBBYSEARCH_SETTARGETUSERID_API_LATEST;
+		SetParameterOptions.Parameter = &AttributeData;
+		SetParameterOptions.ComparisonOp = detail::TranslateSearchComparison(Filter.ComparisonOp);
+
+		EOSResult = EOS_LobbySearch_SetParameter(SearchHandle->Get(), &SetParameterOptions);
+		if (EOSResult != EOS_EResult::EOS_Success)
+		{
+			// todo: errors
+			return MakeFulfilledPromise<TDefaultErrorResultInternal<TSharedPtr<FLobbySearchEOS>>>(FromEOSError(EOSResult)).GetFuture();
+		}
+	}
+
+	TPromise<TDefaultErrorResultInternal<TSharedPtr<FLobbySearchEOS>>> Promise;
+	TFuture<TDefaultErrorResultInternal<TSharedPtr<FLobbySearchEOS>>> Future = Promise.GetFuture();
+
+	EOS_LobbySearch_FindOptions FindOptions = {};
+	FindOptions.ApiVersion = EOS_LOBBYSEARCH_FIND_API_LATEST;
+	FindOptions.LocalUserId = GetProductUserIdChecked(Params.LocalUserId);
+
+	EOS_Async<EOS_LobbySearch_FindCallbackInfo>(EOS_LobbySearch_Find, SearchHandle->Get(), FindOptions)
+	.Then([Promise = MoveTemp(Promise), Prerequisites, LobbyRegistry, LocalUserId = Params.LocalUserId, SearchHandle]
+	(TFuture<const EOS_LobbySearch_FindCallbackInfo*>&& Future) mutable
+	{
+		if (Future.Get()->ResultCode != EOS_EResult::EOS_Success)
+		{
+			// todo: errors
+			Promise.EmplaceValue(FromEOSError(Future.Get()->ResultCode));
+			return;
+		}
+
+		TArray<TFuture<TDefaultErrorResultInternal<TSharedPtr<FLobbyDataEOS>>>> ResolvedLobbyDetails;
+
+		EOS_LobbySearch_GetSearchResultCountOptions GetSearchResultCountOptions = {};
+		GetSearchResultCountOptions.ApiVersion = EOS_LOBBYSEARCH_GETSEARCHRESULTCOUNT_API_LATEST;
+		const uint32_t NumSearchResults = EOS_LobbySearch_GetSearchResultCount(SearchHandle->Get(), &GetSearchResultCountOptions);
+
+		for (uint32_t SearchResultIndex = 0; SearchResultIndex < NumSearchResults; ++SearchResultIndex)
+		{
+			TDefaultErrorResultInternal<TSharedPtr<FLobbyDetailsEOS>> Result = FLobbyDetailsEOS::CreateFromSearchResult(
+				Prerequisites, SearchHandle->Get(), SearchResultIndex);
+			if (Result.IsError())
+			{
+				// todo: errors
+				Promise.EmplaceValue(MoveTemp(Result.GetErrorValue()));
+				return;
+			}
+
+			TPromise<TDefaultErrorResultInternal<TSharedPtr<FLobbyDataEOS>>> ResolveLobbyDetailsPromise;
+			ResolvedLobbyDetails.Add(ResolveLobbyDetailsPromise.GetFuture());
+
+			LobbyRegistry->FindOrCreateFromLobbyDetails(LocalUserId, Result.GetOkValue().ToSharedRef())
+			.Then([ResolveLobbyDetailsPromise = MoveTemp(ResolveLobbyDetailsPromise)](TFuture<TDefaultErrorResultInternal<TSharedPtr<FLobbyDataEOS>>>&& Future) mutable
+			{
+				ResolveLobbyDetailsPromise.EmplaceValue(MoveTempIfPossible(Future.Get()));
+			});
+		}
+
+		WhenAll(MoveTemp(ResolvedLobbyDetails))
+		.Then([Promise = MoveTemp(Promise), SearchHandle](TFuture<TArray<TDefaultErrorResultInternal<TSharedPtr<FLobbyDataEOS>>>> && Future) mutable
+		{
+			TArray<TDefaultErrorResultInternal<TSharedPtr<FLobbyDataEOS>>> Results(Future.Get());
+			TArray<TSharedPtr<FLobbyDataEOS>> ResolvedResults;
+			ResolvedResults.Reserve(Results.Num());
+
+			for (TDefaultErrorResultInternal<TSharedPtr<FLobbyDataEOS>>& Result : Results)
+			{
+				if (Result.IsError())
+				{
+					// todo: errors
+					Promise.EmplaceValue(MoveTemp(Result.GetErrorValue()));
+					return;
+				}
+
+				ResolvedResults.Add(MoveTemp(Result.GetOkValue()));
+			}
+
+			Promise.EmplaceValue(MakeShared<FLobbySearchEOS>(SearchHandle, MoveTemp(ResolvedResults)));
+		});
+	});
+
+	return Future;
+}
+
+TArray<TSharedRef<const FLobby>> FLobbySearchEOS::GetLobbyResults() const
+{
+	TArray<TSharedRef<const FLobby>> Result;
+	Result.Reserve(Lobbies.Num());
+
+	for (const TSharedPtr<FLobbyDataEOS>& LobbyData : Lobbies)
+	{
+		Result.Add(LobbyData->GetLobbyImpl());
+	}
+
+	return Result;
+}
+
+const TArray<TSharedPtr<FLobbyDataEOS>>& FLobbySearchEOS::GetLobbyData()
+{
+	return Lobbies;
+}
+
+FLobbySearchEOS::FLobbySearchEOS(const TSharedRef<FSearchHandle>& SearchHandle, TArray<TSharedPtr<FLobbyDataEOS>>&& Lobbies)
+	: SearchHandle(SearchHandle)
+	, Lobbies(Lobbies)
+{
+}
+
+/* UE::Online */ }
\ No newline at end of file
diff --git a/Engine/Plugins/Online/Experimental/OnlineServicesEOS/Source/Private/LobbiesEOSTypes.h b/Engine/Plugins/Online/Experimental/OnlineServicesEOS/Source/Private/LobbiesEOSTypes.h
new file mode 100644
--- /dev/null
+++ b/Engine/Plugins/Online/Experimental/OnlineServicesEOS/Source/Private/LobbiesEOSTypes.h
@@ -0,0 +1,319 @@
+// Copyright Epic Games, Inc. All Rights Reserved.
+
+#pragma once
+
+#include "CoreMinimal.h"
+#include "Online/Lobbies.h"
+#include "OnlineIdEOS.h"
+#include "OnlineServicesEOSTypes.h"
+
+#if defined(EOS_PLATFORM_BASE_FILE_NAME)
+#include EOS_PLATFORM_BASE_FILE_NAME
+#endif
+
+#include "eos_lobby.h"
+
+namespace UE::Online {
+
+class FLobbySchemaRegistry;
+class FLobbyDataRegistryEOS;
+class FLobbyDetailsInfoEOS;
+class FLobbyInviteDataEOSRegistry;
+class FLobbySchema;
+class FAuthEOS;
+
+inline EOS_ELobbyPermissionLevel TranslateJoinPolicy(ELobbyJoinPolicy JoinPolicy)
+{
+	switch (JoinPolicy)
+	{
+	case ELobbyJoinPolicy::PublicAdvertised:	return EOS_ELobbyPermissionLevel::EOS_LPL_PUBLICADVERTISED;
+	case ELobbyJoinPolicy::PublicNotAdvertised:	return EOS_ELobbyPermissionLevel::EOS_LPL_JOINVIAPRESENCE;
+	default:									checkNoEntry(); // Intentional fallthrough
+	case ELobbyJoinPolicy::InvitationOnly:		return EOS_ELobbyPermissionLevel::EOS_LPL_INVITEONLY;
+	};
+}
+
+using FLobbySearchParameters = FFindLobby::Params;
+
+struct FLobbyDataChanges
+{
+	TOptional<ELobbyJoinPolicy> JoinPolicy;
+
+	// New or changed attributes.
+	TMap<FLobbyAttributeId, FLobbyVariant> MutatedAttributes;
+
+	// Attributes to be cleared.
+	TSet<FLobbyAttributeId> ClearedAttributes;
+};
+
+struct FLobbyMemberDataChanges
+{
+	// New or changed attributes.
+	TMap<FLobbyAttributeId, FLobbyVariant> MutatedAttributes;
+
+	// Attributes to be cleared.
+	TSet<FLobbyAttributeId> ClearedAttributes;
+};
+
+struct FLobbyPrerequisitesEOS
+{
+	EOS_HLobby LobbyInterfaceHandle = {};
+	TSharedRef<FAuthEOS> AuthInterface;
+	TSharedRef<const FLobbySchemaRegistry> SchemaRegistry;
+	TSharedRef<const FLobbySchema> ServiceSchema;
+};
+
+enum class ELobbyDetailsSource
+{
+	Active,
+	Invite,
+	UiEvent,
+	Search
+};
+
+class FLobbyDetailsEOS final : public TSharedFromThis<FLobbyDetailsEOS>
+{
+public:
+	UE_NONCOPYABLE(FLobbyDetailsEOS);
+
+	static TDefaultErrorResultInternal<TSharedPtr<FLobbyDetailsEOS>> CreateFromLobbyId(const TSharedRef<FLobbyPrerequisitesEOS>& Prerequisites, EOS_LobbyId LobbyId, FOnlineAccountIdHandle MemberHandle);
+	static TDefaultErrorResultInternal<TSharedPtr<FLobbyDetailsEOS>> CreateFromInviteId(const TSharedRef<FLobbyPrerequisitesEOS>& Prerequisites, const char* InviteId);
+	static TDefaultErrorResultInternal<TSharedPtr<FLobbyDetailsEOS>> CreateFromUiEventId(const TSharedRef<FLobbyPrerequisitesEOS>& Prerequisites, EOS_UI_EventId UiEventId);
+	static TDefaultErrorResultInternal<TSharedPtr<FLobbyDetailsEOS>> CreateFromSearchResult(const TSharedRef<FLobbyPrerequisitesEOS>& Prerequisites, EOS_HLobbySearch SearchHandle, uint32_t ResultIndex);
+
+	FLobbyDetailsEOS() = delete;
+	~FLobbyDetailsEOS();
+
+	bool IsValid() const { return LobbyDetailsHandle != InvalidLobbyDetailsHandle; }
+	EOS_HLobbyDetails GetEOSHandle() const { return LobbyDetailsHandle; }
+	const TSharedRef<FLobbyDetailsInfoEOS>& GetInfo() const { return LobbyDetailsInfo; }
+	ELobbyDetailsSource GetDetailsSource() const { return LobbyDetailsSource; }
+
+	TFuture<TDefaultErrorResultInternal<TSharedPtr<FLobby>>> GetLobbyData(FOnlineAccountIdHandle LocalUserId, FOnlineLobbyIdHandle LobbyIdHandle) const;
+	TDefaultErrorResultInternal<TSharedPtr<FLobbyMember>> GetLobbyMemberData(FOnlineAccountIdHandle MemberHandle) const;
+	TFuture<EOS_EResult> ApplyLobbyDataUpdates(FOnlineAccountIdHandle LocalUserId, FLobbyDataChanges Changes) const;
+	TFuture<EOS_EResult> ApplyLobbyMemberDataUpdates(FOnlineAccountIdHandle LocalUserId, FLobbyMemberDataChanges Changes) const;
+
+private:
+	template <typename, ESPMode>
+	friend class SharedPointerInternals::TIntrusiveReferenceController;
+	FLobbyDetailsEOS(const TSharedRef<FLobbyPrerequisitesEOS>& Prerequisites, const TSharedRef<FLobbyDetailsInfoEOS>& LobbyDetailsInfo, ELobbyDetailsSource LobbyDetailsSource, EOS_HLobbyDetails LobbyDetailsHandle);
+
+	TPair<FLobbyAttributeId, FLobbyVariant> TranslateLobbyAttribute(const EOS_Lobby_Attribute& LobbyAttribute) const;
+	TDefaultErrorResultInternal<TMap<FLobbyAttributeId, FLobbyVariant>> GetLobbyAttributes() const;
+	TDefaultErrorResultInternal<TMap<FLobbyAttributeId, FLobbyVariant>> GetLobbyMemberAttributes(EOS_ProductUserId TargetMemberProductUserId) const;
+	TSharedRef<FLobbyMember> CreateLobbyMember(FOnlineAccountIdHandle MemberHandle, TMap<FLobbyAttributeId, FLobbyVariant> Attributes) const;
+
+	static const EOS_HLobbyDetails InvalidLobbyDetailsHandle;
+	TSharedRef<FLobbyPrerequisitesEOS> Prerequisites;
+	TSharedRef<FLobbyDetailsInfoEOS> LobbyDetailsInfo;
+	EOS_HLobbyDetails LobbyDetailsHandle = InvalidLobbyDetailsHandle;
+	ELobbyDetailsSource LobbyDetailsSource;
+};
+
+class FLobbyDetailsInfoEOS final
+{
+public:
+	UE_NONCOPYABLE(FLobbyDetailsInfoEOS);
+	FLobbyDetailsInfoEOS() = default;
+
+	static TDefaultErrorResultInternal<TSharedPtr<FLobbyDetailsInfoEOS>> Create(EOS_HLobbyDetails LobbyDetailsHandle);
+
+	bool IsValid() const { return LobbyDetailsInfo.IsValid(); }
+
+	EOS_LobbyId GetLobbyIdEOS() const { return LobbyDetailsInfo->LobbyId; }
+
+private:
+	struct FEOSLobbyDetailsInfoDeleter
+	{
+		void operator()(EOS_LobbyDetails_Info* Ptr) const
+		{
+			EOS_LobbyDetails_Info_Release(Ptr);
+		}
+	};
+
+	TUniquePtr<EOS_LobbyDetails_Info, FEOSLobbyDetailsInfoDeleter> LobbyDetailsInfo;
+};
+
+class FLobbyNotificationPauseHandle
+{
+public:
+	UE_NONCOPYABLE(FLobbyNotificationPauseHandle);
+	FLobbyNotificationPauseHandle() = default;
+	virtual ~FLobbyNotificationPauseHandle() = default;
+};
+
+class FLobbyDataEOS final : public TSharedFromThis<FLobbyDataEOS>
+{
+public:
+	using FUnregisterFn = TFunction<void(FOnlineLobbyIdHandle)>;
+
+	FLobbyDataEOS() = delete;
+
+	static TFuture<TDefaultErrorResultInternal<TSharedPtr<FLobbyDataEOS>>> Create(
+		FOnlineAccountIdHandle LocalUserId,
+		FOnlineLobbyIdHandle LobbyIdHandle,
+		const TSharedRef<FLobbyDetailsEOS>& LobbyDetails,
+		FUnregisterFn UnregisterFn = FUnregisterFn());
+
+	~FLobbyDataEOS();
+
+	FOnlineLobbyIdHandle GetLobbyIdHandle() const { return LobbyImpl->LobbyId; }
+	const TSharedRef<FLobby>& GetLobbyImpl() const { return LobbyImpl; }
+	EOS_LobbyId GetLobbyIdEOS() const { return LobbyDetailsInfo->GetLobbyIdEOS(); }
+	const FString& GetLobbyId() const { return LobbyId; }
+	TSet<FOnlineAccountIdHandle>& GetLocalMembers() { return LocalMembers; }
+
+	// Pause lobby notifications for the target local user.
+	TSharedRef<FLobbyNotificationPauseHandle> PauseLobbyNotifications(FOnlineAccountIdHandle LocalUserId);
+
+	// Returns false if there is an existing notification pause for the target local user.
+	bool CanNotify(FOnlineAccountIdHandle LocalUserId);
+
+	void SetUserLobbyDetails(FOnlineAccountIdHandle LocalUserId, const TSharedPtr<FLobbyDetailsEOS>& LobbyDetails);
+	TSharedPtr<FLobbyDetailsEOS> GetUserLobbyDetails(FOnlineAccountIdHandle LocalUserId) const;
+
+private:
+	template <typename, ESPMode>
+	friend class SharedPointerInternals::TIntrusiveReferenceController;
+
+	FLobbyDataEOS(
+		FOnlineLobbyIdHandle LobbyIdHandle,
+		const TSharedRef<FLobby>& LobbyImpl,
+		const TSharedRef<FLobbyDetailsInfoEOS>& LobbyDetailsInfo,
+		FUnregisterFn UnregisterFn);
+
+	FOnlineLobbyIdHandle LobbyIdHandle;
+	TSharedRef<FLobby> LobbyImpl;
+	TSharedRef<FLobbyDetailsInfoEOS> LobbyDetailsInfo;
+	FUnregisterFn UnregisterFn;
+	FString LobbyId;
+	TSet<FOnlineAccountIdHandle> LocalMembers;
+	TMap<FOnlineAccountIdHandle, TSharedRef<FLobbyNotificationPauseHandle>> NotificationPauses;
+	TMap<FOnlineAccountIdHandle, TSharedPtr<FLobbyDetailsEOS>> UserLobbyDetails;
+};
+
+// Todo: implement handling for FOnlineIdRegistryRegistry.
+class FLobbyDataRegistryEOS : public TSharedFromThis<FLobbyDataRegistryEOS>
+{
+public:
+	FLobbyDataRegistryEOS(const TSharedRef<FLobbyPrerequisitesEOS>& Prerequisites);
+
+	TSharedPtr<FLobbyDataEOS> Find(EOS_LobbyId EOSLobbyId) const;
+	TSharedPtr<FLobbyDataEOS> Find(FOnlineLobbyIdHandle LobbyIdHandle) const;
+	TFuture<TDefaultErrorResultInternal<TSharedPtr<FLobbyDataEOS>>> FindOrCreateFromLobbyId(FOnlineAccountIdHandle LocalUserId, EOS_LobbyId LobbyId);
+	TFuture<TDefaultErrorResultInternal<TSharedPtr<FLobbyDataEOS>>> FindOrCreateFromUiEventId(FOnlineAccountIdHandle LocalUserId, EOS_UI_EventId UiEventId);
+	TFuture<TDefaultErrorResultInternal<TSharedPtr<FLobbyDataEOS>>> FindOrCreateFromLobbyDetails(FOnlineAccountIdHandle LocalUserId, const TSharedRef<FLobbyDetailsEOS>& LobbyDetails);
+
+private:
+	void Register(const TSharedRef<FLobbyDataEOS>& LobbyIdHandleData);
+	void Unregister(FOnlineLobbyIdHandle LobbyIdHandle);
+	FLobbyDataEOS::FUnregisterFn MakeUnregisterFn();
+
+	TSharedRef<FLobbyPrerequisitesEOS> Prerequisites;
+	TMap<FString, TWeakPtr<FLobbyDataEOS>> LobbyIdIndex;
+	TMap<FOnlineLobbyIdHandle, TWeakPtr<FLobbyDataEOS>> LobbyIdHandleIndex;
+	uint32 NextHandleIndex = 1;
+};
+
+class FLobbyInviteIdEOS final
+{
+public:
+	FLobbyInviteIdEOS() = default;
+	FLobbyInviteIdEOS(const char* InviteIdEOS)
+	{
+		FPlatformString::Strcpy(Data, MaxLobbyInviteIdSize, InviteIdEOS);
+	}
+
+	const char* Get() const { return Data; }
+
+private:
+	static const int MaxLobbyInviteIdSize = 256;
+	char Data[MaxLobbyInviteIdSize] = {};
+};
+
+class FLobbyInviteDataEOS final
+{
+public:
+	FLobbyInviteDataEOS() = delete;
+
+	static TFuture<TDefaultErrorResultInternal<TSharedPtr<FLobbyInviteDataEOS>>> CreateFromInviteId(
+		const TSharedRef<FLobbyPrerequisitesEOS>& Prerequisites,
+		const TSharedRef<FLobbyDataRegistryEOS>& LobbyDataRegistry,
+		FOnlineAccountIdHandle LocalUserId,
+		const char* InviteIdEOS,
+		EOS_ProductUserId Sender);
+
+	TSharedRef<FLobbyDetailsEOS> GetLobbyDetails() const { return LobbyDetails; }
+	TSharedRef<FLobbyDataEOS> GetLobbyData() const { return LobbyData; }
+
+	const char* GetInviteIdEOS() const { return InviteIdEOS->Get(); }
+	const FString& GetInviteId() const { return InviteId; }
+
+	FOnlineAccountIdHandle GetReceiver() const { return Receiver; }
+	FOnlineAccountIdHandle GetSender() const { return Sender; }
+
+private:
+	template <typename, ESPMode>
+	friend class SharedPointerInternals::TIntrusiveReferenceController;
+
+	FLobbyInviteDataEOS(
+		const TSharedRef<FLobbyInviteIdEOS>& InviteIdEOS,
+		FOnlineAccountIdHandle Receiver,
+		FOnlineAccountIdHandle Sender,
+		const TSharedRef<FLobbyDetailsEOS>& LobbyDetails,
+		const TSharedRef<FLobbyDataEOS>& LobbyData);
+
+	TSharedRef<FLobbyInviteIdEOS> InviteIdEOS;
+	FOnlineAccountIdHandle Receiver;
+	FOnlineAccountIdHandle Sender;
+	TSharedRef<FLobbyDetailsEOS> LobbyDetails;
+	TSharedRef<FLobbyDataEOS> LobbyData;
+	FString InviteId;
+};
+
+// The lobby search object is meant to hold the lifetimes of lobby search results so that the
+// lobby details will be discoverable during a join operation.
+class FLobbySearchEOS final
+{
+public:
+	FLobbySearchEOS() = delete;
+
+	static TFuture<TDefaultErrorResultInternal<TSharedPtr<FLobbySearchEOS>>> Create(
+		const TSharedRef<FLobbyPrerequisitesEOS>& Prerequisites,
+		const TSharedRef<FLobbyDataRegistryEOS>& LobbyRegistry,
+		const FLobbySearchParameters& Params);
+
+	TArray<TSharedRef<const FLobby>> GetLobbyResults() const;
+	const TArray<TSharedPtr<FLobbyDataEOS>>& GetLobbyData();
+
+private:
+	class FSearchHandle
+	{
+	public:
+		UE_NONCOPYABLE(FSearchHandle);
+		FSearchHandle() = default;
+
+		~FSearchHandle()
+		{
+			EOS_LobbySearch_Release(SearchHandle);
+		}
+
+		EOS_HLobbySearch& Get()
+		{
+			return SearchHandle;
+		}
+
+	private:
+		EOS_HLobbySearch SearchHandle = {};
+	};
+
+	template <typename, ESPMode>
+	friend class SharedPointerInternals::TIntrusiveReferenceController;
+	FLobbySearchEOS(const TSharedRef<FSearchHandle>& SearchHandle, TArray<TSharedPtr<FLobbyDataEOS>>&& Lobbies);
+
+	TSharedRef<FSearchHandle> SearchHandle;
+	TArray<TSharedPtr<FLobbyDataEOS>> Lobbies;
+};
+
+/* UE::Online */ }
diff --git a/Engine/Plugins/Online/Experimental/OnlineServicesEOS/Source/Public/LobbiesEOS.h b/Engine/Plugins/Online/Experimental/OnlineServicesEOS/Source/Public/LobbiesEOS.h
new file mode 100644
--- /dev/null
+++ b/Engine/Plugins/Online/Experimental/OnlineServicesEOS/Source/Public/LobbiesEOS.h
@@ -0,0 +1,308 @@
+// Copyright Epic Games, Inc. All Rights Reserved.
+
+#pragma once
+
+#include "CoreMinimal.h"
+#include "Online/LobbiesCommon.h"
+#include "OnlineServicesEOSTypes.h"
+
+#if defined(EOS_PLATFORM_BASE_FILE_NAME)
+#include EOS_PLATFORM_BASE_FILE_NAME
+#endif
+#include "eos_lobby_types.h"
+
+namespace UE::Online {
+
+struct FLobbyPrerequisitesEOS;
+class FLobbyDataEOS;
+class FLobbyDataRegistryEOS;
+class FLobbyDetailsEOS;
+class FLobbyInviteDataEOS;
+class FOnlineServicesEOS;
+class FLobbySearchEOS;
+struct FLobbyDataChanges;
+struct FLobbyMemberDataChanges;
+	
+class FLobbiesEOS : public FLobbiesCommon
+{
+public:
+	FLobbiesEOS(FOnlineServicesEOS& InServices);
+
+	virtual void Initialize() override;
+	virtual void PreShutdown() override;
+
+	virtual TOnlineAsyncOpHandle<FCreateLobby> CreateLobby(FCreateLobby::Params&& Params) override;
+	virtual TOnlineAsyncOpHandle<FFindLobby> FindLobby(FFindLobby::Params&& Params) override;
+	virtual TOnlineAsyncOpHandle<FJoinLobby> JoinLobby(FJoinLobby::Params&& Params) override;
+	virtual TOnlineAsyncOpHandle<FLeaveLobby> LeaveLobby(FLeaveLobby::Params&& Params) override;
+	virtual TOnlineAsyncOpHandle<FInviteLobbyMember> InviteLobbyMember(FInviteLobbyMember::Params&& Params) override;
+	virtual TOnlineAsyncOpHandle<FDeclineLobbyInvitation> DeclineLobbyInvitation(FDeclineLobbyInvitation::Params&& Params) override;
+	virtual TOnlineAsyncOpHandle<FKickLobbyMember> KickLobbyMember(FKickLobbyMember::Params&& Params) override;
+	virtual TOnlineAsyncOpHandle<FPromoteLobbyMember> PromoteLobbyMember(FPromoteLobbyMember::Params&& Params) override;
+	virtual TOnlineAsyncOpHandle<FModifyLobbyJoinPolicy> ModifyLobbyJoinPolicy(FModifyLobbyJoinPolicy::Params&& Params) override;
+	virtual TOnlineAsyncOpHandle<FModifyLobbyAttributes> ModifyLobbyAttributes(FModifyLobbyAttributes::Params&& Params) override;
+	virtual TOnlineAsyncOpHandle<FModifyLobbyMemberAttributes> ModifyLobbyMemberAttributes(FModifyLobbyMemberAttributes::Params&& Params) override;
+
+private:
+	void HandleLobbyUpdated(const EOS_Lobby_LobbyUpdateReceivedCallbackInfo* Data);
+	void HandleLobbyMemberUpdated(const EOS_Lobby_LobbyMemberUpdateReceivedCallbackInfo* Data);
+	void HandleLobbyMemberStatusReceived(const EOS_Lobby_LobbyMemberStatusReceivedCallbackInfo* Data);
+	void HandleLobbyInviteReceived(const EOS_Lobby_LobbyInviteReceivedCallbackInfo* Data);
+	void HandleLobbyInviteAccepted(const EOS_Lobby_LobbyInviteAcceptedCallbackInfo* Data);
+	void HandleJoinLobbyAccepted(const EOS_Lobby_JoinLobbyAcceptedCallbackInfo* Data);
+
+protected:
+	void RegisterHandlers();
+	void UnregisterHandlers();
+
+	void AddActiveLobby(FOnlineAccountIdHandle LocalUserId, const TSharedRef<FLobbyDataEOS>& LobbyData);
+	void RemoveActiveLobby(FOnlineAccountIdHandle LocalUserId, const TSharedRef<FLobbyDataEOS>& LobbyData);
+
+	// Todo: store list of invites per lobby.
+	void AddActiveInvite(const TSharedRef<FLobbyInviteDataEOS>& Invite);
+	void RemoveActiveInvite(const TSharedRef<FLobbyInviteDataEOS>& Invite);
+	TSharedPtr<FLobbyInviteDataEOS> GetActiveInvite(FOnlineAccountIdHandle TargetUser, FOnlineLobbyIdHandle TargetLobbyId);
+
+	struct FJoinLobbyImpl
+	{
+		static constexpr TCHAR Name[] = TEXT("JoinLobbyImpl");
+
+		struct Params : public FJoinLobby::Params
+		{
+			// The lobby handle data.
+			TSharedPtr<FLobbyDataEOS> LobbyData;
+
+			// The local user agent which will perform the action.
+			FOnlineAccountIdHandle LocalUserId;
+
+			// The local name for the lobby.
+			FName LocalName;
+
+			// Local users who will be joining the lobby.
+			TArray<FJoinLobbyLocalUserData> LocalUsers;
+		};
+
+		struct Result
+		{
+		};
+	};
+
+	// LobbyData will be fetched from the operation data if not set in Params.
+	TFuture<TDefaultErrorResult<FJoinLobbyImpl>> JoinLobbyImpl(FJoinLobbyImpl::Params&& Params);
+
+	struct FJoinLobbyMemberImpl
+	{
+		static constexpr TCHAR Name[] = TEXT("JoinLobbyMemberImpl");
+
+		struct Params
+		{
+			// The lobby handle data.
+			TSharedPtr<FLobbyDataEOS> LobbyData;
+
+			// The local user agent which will perform the action.
+			FOnlineAccountIdHandle LocalUserId;
+
+			// Initial attributes.
+			TMap<FLobbyAttributeId, FLobbyVariant> Attributes;
+		};
+
+		struct Result
+		{
+		};
+	};
+
+	TOnlineAsyncOpHandle<FJoinLobbyMemberImpl> JoinLobbyMemberImplOp(FJoinLobbyMemberImpl::Params&& Params);
+
+	struct FLeaveLobbyImpl
+	{
+		static constexpr TCHAR Name[] = TEXT("LeaveLobbyImpl");
+
+		struct Params
+		{
+			// The lobby handle data.
+			TSharedPtr<FLobbyDataEOS> LobbyData;
+
+			// The local user agent which will perform the action.
+			FOnlineAccountIdHandle LocalUserId;
+		};
+
+		struct Result
+		{
+		};
+	};
+
+	TFuture<TDefaultErrorResult<FLeaveLobbyImpl>> LeaveLobbyImpl(FLeaveLobbyImpl::Params&& Params);
+
+	struct FDestroyLobbyImpl
+	{
+		static constexpr TCHAR Name[] = TEXT("DestroyLobbyImpl");
+
+		struct Params
+		{
+			// The lobby handle data.
+			TSharedPtr<FLobbyDataEOS> LobbyData;
+
+			// The local user agent which will perform the action.
+			FOnlineAccountIdHandle LocalUserId;
+		};
+
+		struct Result
+		{
+		};
+	};
+
+	TFuture<TDefaultErrorResult<FDestroyLobbyImpl>> DestroyLobbyImpl(FDestroyLobbyImpl::Params&& Params);
+
+	struct FInviteLobbyMemberImpl
+	{
+		static constexpr TCHAR Name[] = TEXT("InviteLobbyMemberImpl");
+
+		struct Params
+		{
+			// The lobby handle data.
+			TSharedPtr<FLobbyDataEOS> LobbyData;
+
+			// The local user agent which will perform the action.
+			FOnlineAccountIdHandle LocalUserId;
+
+			// The target user for the invitation.
+			FOnlineAccountIdHandle TargetUserId;
+		};
+
+		struct Result
+		{
+		};
+	};
+
+	TFuture<TDefaultErrorResult<FInviteLobbyMemberImpl>> InviteLobbyMemberImpl(FInviteLobbyMemberImpl::Params&& Params);
+
+	struct FDeclineLobbyInvitationImpl
+	{
+		static constexpr TCHAR Name[] = TEXT("DeclineLobbyInvitationImpl");
+
+		struct Params
+		{
+			// The local user agent which will perform the action.
+			FOnlineAccountIdHandle LocalUserId;
+
+			// Id of the lobby for which the invitations will be declined.
+			FOnlineLobbyIdHandle LobbyId;
+		};
+
+		struct Result
+		{
+		};
+	};
+
+	TFuture<TDefaultErrorResult<FDeclineLobbyInvitationImpl>> DeclineLobbyInvitationImpl(FDeclineLobbyInvitationImpl::Params&& Params);
+
+	struct FKickLobbyMemberImpl
+	{
+		static constexpr TCHAR Name[] = TEXT("KickLobbyMemberImpl");
+
+		struct Params
+		{
+			// The lobby handle data.
+			TSharedPtr<FLobbyDataEOS> LobbyData;
+
+			// The local user agent which will perform the action.
+			FOnlineAccountIdHandle LocalUserId;
+
+			// The target user to be kicked.
+			FOnlineAccountIdHandle TargetUserId;
+		};
+
+		struct Result
+		{
+		};
+	};
+
+	TFuture<TDefaultErrorResult<FKickLobbyMemberImpl>> KickLobbyMemberImpl(FKickLobbyMemberImpl::Params&& Params);
+
+	struct FPromoteLobbyMemberImpl
+	{
+		static constexpr TCHAR Name[] = TEXT("PromoteLobbyMemberImpl");
+
+		struct Params
+		{
+			// The lobby handle data.
+			TSharedPtr<FLobbyDataEOS> LobbyData;
+
+			// The local user agent which will perform the action.
+			FOnlineAccountIdHandle LocalUserId;
+
+			// The target user to be promoted to owner.
+			FOnlineAccountIdHandle TargetUserId;
+		};
+
+		struct Result
+		{
+		};
+	};
+
+	TFuture<TDefaultErrorResult<FPromoteLobbyMemberImpl>> PromoteLobbyMemberImpl(FPromoteLobbyMemberImpl::Params&& Params);
+
+	struct FModifyLobbyDataImpl
+	{
+		static constexpr TCHAR Name[] = TEXT("ModifyLobbyDataImpl");
+
+		struct Params
+		{
+			// The lobby handle data.
+			TSharedPtr<FLobbyDataEOS> LobbyData;
+
+			// The local user agent which will perform the action.
+			FOnlineAccountIdHandle LocalUserId;
+
+			// The changes to apply to the lobby data.
+			TSharedPtr<FLobbyDataChanges> Changes;
+		};
+
+		struct Result
+		{
+		};
+	};
+
+	TFuture<TDefaultErrorResult<FModifyLobbyDataImpl>> ModifyLobbyDataImpl(FModifyLobbyDataImpl::Params&& Params);
+
+	struct FModifyLobbyMemberDataImpl
+	{
+		static constexpr TCHAR Name[] = TEXT("ModifyLobbyMemberDataImpl");
+
+		struct Params
+		{
+			// The lobby handle data.
+			TSharedPtr<FLobbyDataEOS> LobbyData;
+
+			// The local user agent which will perform the action.
+			FOnlineAccountIdHandle LocalUserId;
+
+			// The changes to apply to the lobby member data.
+			TSharedPtr<FLobbyMemberDataChanges> Changes;
+		};
+
+		struct Result
+		{
+		};
+	};
+
+	TFuture<TDefaultErrorResult<FModifyLobbyMemberDataImpl>> ModifyLobbyMemberDataImpl(FModifyLobbyMemberDataImpl::Params&& Params);
+
+	FString LobbyBucketId;
+
+	EOSEventRegistrationPtr OnLobbyUpdatedEOSEventRegistration;
+	EOSEventRegistrationPtr OnLobbyMemberUpdatedEOSEventRegistration;
+	EOSEventRegistrationPtr OnLobbyMemberStatusReceivedEOSEventRegistration;
+	EOSEventRegistrationPtr OnLobbyInviteReceivedEOSEventRegistration;
+	EOSEventRegistrationPtr OnLobbyInviteAcceptedEOSEventRegistration;
+	EOSEventRegistrationPtr OnJoinLobbyAcceptedEOSEventRegistration;
+
+	TSharedPtr<FLobbyPrerequisitesEOS> LobbyPrerequisites;
+	TSharedPtr<FLobbyDataRegistryEOS> LobbyDataRegistry;
+
+	TMap<FOnlineAccountIdHandle, TSet<TSharedRef<FLobbyDataEOS>>> ActiveLobbies;
+	TMap<FOnlineAccountIdHandle, TMap<FOnlineLobbyIdHandle, TSharedRef<FLobbyInviteDataEOS>>> ActiveInvites;
+	TMap<FOnlineAccountIdHandle, TSharedRef<FLobbySearchEOS>> ActiveSearchResults;
+};
+
+/* UE::Online */ }
