Fix static analysis issue

diff --git a/Engine/Plugins/Online/Experimental/OnlineServicesEOS/Source/Private/LobbiesEOS.cpp b/Engine/Plugins/Online/Experimental/OnlineServicesEOS/Source/Private/LobbiesEOS.cpp
--- a/Engine/Plugins/Online/Experimental/OnlineServicesEOS/Source/Private/LobbiesEOS.cpp
+++ b/Engine/Plugins/Online/Experimental/OnlineServicesEOS/Source/Private/LobbiesEOS.cpp
@@ -1039,7 +1039,6 @@ TFuture<TDefaultErrorResult<FLobbiesEOS::FJoinLobbyImpl>> FLobbiesEOS::JoinLobby
 			TArray<TFuture<TDefaultErrorResult<FLeaveLobbyImpl>>> PendingMemberExits;
 			for (int32 MemberIndex = 0; MemberIndex < Params.LocalUsers.Num(); ++MemberIndex)
 			{
-				const TDefaultErrorResult<FJoinLobbyMemberImpl>& MemberResult = Results.Get()[MemberIndex];
 				const FOnlineAccountIdHandle MemberId = Params.LocalUsers[MemberIndex].LocalUserId;
 
 				if (Results.Get()[MemberIndex].IsError())
