- Add methods to create FOnlineError from EOS_EResult, and comparison operators
- Add comparison operators for FOnlineError/ErrorCodeType

diff --git a/Engine/Plugins/Online/Experimental/OnlineServices/Source/OnlineServicesInterface/Public/Online/OnlineError.h b/Engine/Plugins/Online/Experimental/OnlineServices/Source/OnlineServicesInterface/Public/Online/OnlineError.h
--- a/Engine/Plugins/Online/Experimental/OnlineServices/Source/OnlineServicesInterface/Public/Online/OnlineError.h
+++ b/Engine/Plugins/Online/Experimental/OnlineServices/Source/OnlineServicesInterface/Public/Online/OnlineError.h
@@ -149,7 +149,11 @@ private:
 
 // the comparison with ErrorCodeType is the one I would expect to see in almost every case
 ONLINESERVICESINTERFACE_API bool operator==(const FOnlineError& Lhs, const FOnlineError& Rhs);
-ONLINESERVICESINTERFACE_API bool operator==(const FOnlineError& Lhs, ErrorCodeType OtherErrorCode);
+inline bool operator!=(const FOnlineError& Lhs, const FOnlineError& Rhs) { return !(Lhs == Rhs); }
+ONLINESERVICESINTERFACE_API bool operator==(const FOnlineError& OnlineError, ErrorCodeType OtherErrorCode);
+inline bool operator==(ErrorCodeType OtherErrorCode, const FOnlineError& OnlineError) { return OnlineError == OtherErrorCode; }
+inline bool operator!=(const FOnlineError& OnlineError, ErrorCodeType OtherErrorCode) { return !(OnlineError == OtherErrorCode); }
+inline bool operator!=(ErrorCodeType OtherErrorCode, const FOnlineError& OnlineError) { return !(OnlineError == OtherErrorCode); }
 
 inline FString ToLogString(const FOnlineError& Error)
 {
diff --git a/Engine/Plugins/Online/Experimental/OnlineServices/Source/OnlineServicesInterface/Public/Online/OnlineErrorCode.h b/Engine/Plugins/Online/Experimental/OnlineServices/Source/OnlineServicesInterface/Public/Online/OnlineErrorCode.h
--- a/Engine/Plugins/Online/Experimental/OnlineServices/Source/OnlineServicesInterface/Public/Online/OnlineErrorCode.h
+++ b/Engine/Plugins/Online/Experimental/OnlineServices/Source/OnlineServicesInterface/Public/Online/OnlineErrorCode.h
@@ -39,9 +39,9 @@ namespace ErrorCode
 // These three functions are responsible for reading the parts of the uint64 error representation
 // The uint64 error is formatted as follows (big endian):
 // ssss cccc cccc cccc cccc cccc cccc cccc cccc vvvv vvvv vvvv vvvv vvvv vvvv vvvv vvvv 
-uint64 ErrorCodeSystem(ErrorCodeType ErrorCode);
-uint64 ErrorCodeCategory(ErrorCodeType ErrorCode);
-uint64 ErrorCodeValue(ErrorCodeType ErrorCode); 
+ONLINESERVICESINTERFACE_API uint64 ErrorCodeSystem(ErrorCodeType ErrorCode);
+ONLINESERVICESINTERFACE_API uint64 ErrorCodeCategory(ErrorCodeType ErrorCode);
+ONLINESERVICESINTERFACE_API uint64 ErrorCodeValue(ErrorCodeType ErrorCode); 
 
 // needs to be inside Online::V2::Errors namespace
 // define an error category
diff --git a/Engine/Plugins/Online/Experimental/OnlineServicesEOS/Source/Private/OnlineErrorEOS.cpp b/Engine/Plugins/Online/Experimental/OnlineServicesEOS/Source/Private/OnlineErrorEOS.cpp
new file mode 100644
--- /dev/null
+++ b/Engine/Plugins/Online/Experimental/OnlineServicesEOS/Source/Private/OnlineErrorEOS.cpp
@@ -0,0 +1,40 @@
+// Copyright Epic Games, Inc. All Rights Reserved.
+
+#include "OnlineErrorEOS.h"
+
+namespace UE::Online {
+
+const TSharedRef<const IOnlineErrorDetails, ESPMode::ThreadSafe>& FOnlineErrorDetailsEOS::Get()
+{
+	static TSharedRef<const IOnlineErrorDetails, ESPMode::ThreadSafe> Instance = MakeShared<FOnlineErrorDetailsEOS>();
+	return Instance;
+}
+
+FText FOnlineErrorDetailsEOS::GetText(const FOnlineError& OnlineError) const
+{
+	const EOS_EResult EosResult = EOS_EResult(OnlineError.GetValue());
+	return FText::FromStringView(FStringView(UTF8_TO_TCHAR(EOS_EResult_ToString(EosResult))));
+}
+
+FString FOnlineErrorDetailsEOS::GetLogString(const FOnlineError& OnlineError) const 
+{
+	const EOS_EResult EosResult = EOS_EResult(OnlineError.GetValue());
+	return UTF8_TO_TCHAR(EOS_EResult_ToString(EosResult));
+}
+
+bool operator==(const UE::Online::FOnlineError& OnlineError, EOS_EResult EosResult)
+{
+	return OnlineError.GetErrorCode() == CreateErrorCode(EosResult);
+}
+
+ErrorCodeType CreateErrorCode(EOS_EResult EosResult)
+{
+	return Errors::ErrorCode::Create(Errors::ErrorCode::Category::EOS_System, Errors::ErrorCode::Category::EOS, uint32(EosResult));
+}
+
+FOnlineError CreateOnlineError(EOS_EResult EosResult)
+{
+	return FOnlineError(CreateErrorCode(EosResult), FOnlineErrorDetailsEOS::Get());
+}
+
+} /* namespace UE::Online */
diff --git a/Engine/Plugins/Online/Experimental/OnlineServicesEOS/Source/Private/OnlineErrorEOS.h b/Engine/Plugins/Online/Experimental/OnlineServicesEOS/Source/Private/OnlineErrorEOS.h
new file mode 100644
--- /dev/null
+++ b/Engine/Plugins/Online/Experimental/OnlineServicesEOS/Source/Private/OnlineErrorEOS.h
@@ -0,0 +1,37 @@
+// Copyright Epic Games, Inc. All Rights Reserved.
+
+#pragma once
+
+#include "Online/OnlineError.h"
+
+#if defined(EOS_PLATFORM_BASE_FILE_NAME)
+#include EOS_PLATFORM_BASE_FILE_NAME
+#endif
+#include "eos_common.h"
+
+namespace UE::Online {
+
+namespace Errors {
+
+UE_ONLINE_ERROR_CATEGORY(EOS, Engine, 0x3, "Online Services EOS")
+
+} /* namespace Errors */
+
+class FOnlineErrorDetailsEOS : public IOnlineErrorDetails
+{
+public:
+	static const TSharedRef<const IOnlineErrorDetails, ESPMode::ThreadSafe>& Get();
+
+	virtual FText GetText(const FOnlineError& OnlineError) const override;
+	virtual FString GetLogString(const FOnlineError& OnlineError) const override;
+};
+
+bool operator==(const FOnlineError& OnlineError, EOS_EResult EosResult);
+inline bool operator==(EOS_EResult EosResult, const FOnlineError& OnlineError) { return OnlineError == EosResult; }
+inline bool operator!=(const FOnlineError& OnlineError, EOS_EResult EosResult) { return !(OnlineError == EosResult); }
+inline bool operator!=(EOS_EResult EosResult, const FOnlineError& OnlineError) { return !(OnlineError == EosResult); }
+
+ErrorCodeType CreateErrorCode(EOS_EResult EosResult);
+FOnlineError CreateOnlineError(EOS_EResult EosResult);
+
+} /* namespace UE::Online */
\ No newline at end of file
