Check if we are already logged in
Fix instance of TSharedRef default constructor

diff --git a/Engine/Plugins/Online/Experimental/OnlineServices/Source/OnlineServicesInterface/Public/Online/AuthErrors.h b/Engine/Plugins/Online/Experimental/OnlineServices/Source/OnlineServicesInterface/Public/Online/AuthErrors.h
--- a/Engine/Plugins/Online/Experimental/OnlineServices/Source/OnlineServicesInterface/Public/Online/AuthErrors.h
+++ b/Engine/Plugins/Online/Experimental/OnlineServices/Source/OnlineServicesInterface/Public/Online/AuthErrors.h
@@ -4,8 +4,12 @@
 
 #include "Online/OnlineError.h"
 
- // RobC mod- whole file
+#define LOCTEXT_NAMESPACE "AuthErrors"
+
 namespace UE::Online::Errors {
-	// see WindowsOnlineErrorDefinitions.h for an example of this when it comes time to develop
+	UE_ONLINE_ERROR_CATEGORY(Auth, Engine, 0x4, "Auth")
+	UE_ONLINE_ERROR(Auth, AlreadyLoggedIn, 1, TEXT("AlreadyLoggedIn"), LOCTEXT("AlreadyLoggedIn", "That user is already logged in"))
 
 /* UE::Online::Errors */ }
+
+#undef LOCTEXT_NAMESPACE
diff --git a/Engine/Plugins/Online/Experimental/OnlineServicesEOS/Source/Private/AuthEOS.cpp b/Engine/Plugins/Online/Experimental/OnlineServicesEOS/Source/Private/AuthEOS.cpp
--- a/Engine/Plugins/Online/Experimental/OnlineServicesEOS/Source/Private/AuthEOS.cpp
+++ b/Engine/Plugins/Online/Experimental/OnlineServicesEOS/Source/Private/AuthEOS.cpp
@@ -339,6 +339,12 @@ TOnlineAsyncOpHandle<FAuthLogin> FAuthEOS::Login(FAuthLogin::Params&& Params)
 		}
 	}
 	TOnlineAsyncOpRef<FAuthLogin> Op = GetOp<FAuthLogin>(MoveTemp(Params));
+	// Are we already logged in?
+	if (GetAccountIdByPlatformUserId(Params.PlatformUserId).IsOk())
+	{
+		Op->SetError(Errors::Auth::AlreadyLoggedIn());
+		return Op->GetHandle();
+	}
 
 	EOS_Auth_LoginOptions LoginOptions = { };
 	LoginOptions.ApiVersion = EOS_AUTH_LOGIN_API_LATEST;
@@ -718,8 +724,7 @@ TOnlineResult<FAuthGetAccountByAccountId> FAuthEOS::GetAccountByAccountId(FAuthG
 {
 	if (TSharedRef<FAccountInfoEOS>* const FoundAccount = AccountInfos.Find(Params.LocalUserId))
 	{
-		FAuthGetAccountByAccountId::Result Result;
-		Result.AccountInfo = *FoundAccount;
+		FAuthGetAccountByAccountId::Result Result = { *FoundAccount };
 		return TOnlineResult<FAuthGetAccountByAccountId>(Result);
 	}
 	else
