Implement GetJoinedLobbies in EOS
Set EOS version on attributes
Set LocalName when creating or joining a session
Log EOS error string instead of error number

diff --git a/Engine/Plugins/Online/Experimental/OnlineServices/Source/OnlineServicesInterface/Public/Online/Lobbies.h b/Engine/Plugins/Online/Experimental/OnlineServices/Source/OnlineServicesInterface/Public/Online/Lobbies.h
--- a/Engine/Plugins/Online/Experimental/OnlineServices/Source/OnlineServicesInterface/Public/Online/Lobbies.h
+++ b/Engine/Plugins/Online/Experimental/OnlineServices/Source/OnlineServicesInterface/Public/Online/Lobbies.h
@@ -849,9 +849,11 @@ BEGIN_ONLINE_STRUCT_META(FModifyLobbyMemberAttributes::Result)
 END_ONLINE_STRUCT_META()
 
 BEGIN_ONLINE_STRUCT_META(FGetJoinedLobbies::Params)
+	ONLINE_STRUCT_FIELD(FGetJoinedLobbies::Params, LocalUserId)
 END_ONLINE_STRUCT_META()
 
 BEGIN_ONLINE_STRUCT_META(FGetJoinedLobbies::Result)
+	ONLINE_STRUCT_FIELD(FGetJoinedLobbies::Result, Lobbies)
 END_ONLINE_STRUCT_META()
 
 BEGIN_ONLINE_STRUCT_META(FGetReceivedInvitations::Params)
diff --git a/Engine/Plugins/Online/Experimental/OnlineServicesEOS/Source/Private/LobbiesEOS.cpp b/Engine/Plugins/Online/Experimental/OnlineServicesEOS/Source/Private/LobbiesEOS.cpp
--- a/Engine/Plugins/Online/Experimental/OnlineServicesEOS/Source/Private/LobbiesEOS.cpp
+++ b/Engine/Plugins/Online/Experimental/OnlineServicesEOS/Source/Private/LobbiesEOS.cpp
@@ -230,6 +230,7 @@ TOnlineAsyncOpHandle<FCreateLobby> FLobbiesEOS::CreateLobby(FCreateLobby::Params
 	{
 		const auto& Params = InAsyncOp.GetParams();
 		auto LobbyData = GetOpDataChecked<TSharedRef<FLobbyDataEOS>>(InAsyncOp, LobbyDataKeyName);
+		LobbyData->GetLobbyImpl()->LocalName = Params.LocalName;
 
 		// Remove notification pause for joining local users.
 		TArray<TSharedRef<FLobbyNotificationPauseHandle>> ClearedPauses;
@@ -333,6 +334,7 @@ TOnlineAsyncOpHandle<FJoinLobby> FLobbiesEOS::JoinLobby(FJoinLobby::Params&& InP
 		}
 		else
 		{
+			LobbyData->GetLobbyImpl()->LocalName = InAsyncOp.GetParams().LocalName;
 			// Remove notification pause for joining local users.
 			TArray<TSharedRef<FLobbyNotificationPauseHandle>> ClearedPauses;
 			InAsyncOp.Data.Set<TArray<TSharedRef<FLobbyNotificationPauseHandle>>>(LobbyNotificationPauseKeyName, MoveTemp(ClearedPauses));
@@ -344,7 +346,7 @@ TOnlineAsyncOpHandle<FJoinLobby> FLobbiesEOS::JoinLobby(FJoinLobby::Params&& InP
 				AddActiveLobby(LocalUser.LocalUserId, LobbyData.ToSharedRef());
 			}
 
-			InAsyncOp.SetResult(FJoinLobby::Result{});
+			InAsyncOp.SetResult(FJoinLobby::Result{ LobbyData->GetLobbyImpl()});
 		}
 	})
 	.Enqueue(GetSerialQueue());
@@ -673,6 +675,24 @@ TOnlineAsyncOpHandle<FModifyLobbyMemberAttributes> FLobbiesEOS::ModifyLobbyMembe
 	return Op->GetHandle();
 }
 
+TOnlineResult<FGetJoinedLobbies> FLobbiesEOS::GetJoinedLobbies(FGetJoinedLobbies::Params&& Params)
+{
+	if (const TSet<TSharedRef<FLobbyDataEOS>>* Lobbies = ActiveLobbies.Find(Params.LocalUserId))
+	{
+		FGetJoinedLobbies::Result Result;
+		Result.Lobbies.Reserve(Lobbies->Num());
+		for (const TSharedRef<FLobbyDataEOS>& LobbyDataEOS : *Lobbies)
+		{
+			Result.Lobbies.Emplace(LobbyDataEOS->GetLobbyImpl());
+		}
+		return TOnlineResult<FGetJoinedLobbies>(MoveTemp(Result));
+	}
+	else
+	{
+		return TOnlineResult<FGetJoinedLobbies>(Errors::InvalidUser());
+	}
+}
+
 void FLobbiesEOS::HandleLobbyUpdated(const EOS_Lobby_LobbyUpdateReceivedCallbackInfo* Data)
 {
 #if 0
@@ -1443,7 +1463,7 @@ TFuture<TDefaultErrorResult<FLobbiesEOS::FModifyLobbyDataImpl>> FLobbiesEOS::Mod
 	DetailsResult.GetOkValue()->ApplyLobbyDataUpdates(Params.LocalUserId, MoveTemp(*Params.Changes))
 	.Then([Promise = MoveTemp(Promise)](TFuture<EOS_EResult>&& Future) mutable
 	{
-		UE_LOG(LogTemp, Warning, TEXT("[FLobbiesEOS::ModifyLobbyDataImpl] Complete. Result: %d"), Future.Get());
+		UE_LOG(LogTemp, Warning, TEXT("[FLobbiesEOS::ModifyLobbyDataImpl] Complete. Result: %s"), *LexToString(Future.Get()));
 
 		// Todo: Handle "no change" better.
 		if (Future.Get() != EOS_EResult::EOS_Success && Future.Get() != EOS_EResult::EOS_NoChange)
@@ -1489,7 +1509,7 @@ TFuture<TDefaultErrorResult<FLobbiesEOS::FModifyLobbyMemberDataImpl>> FLobbiesEO
 	DetailsResult.GetOkValue()->ApplyLobbyMemberDataUpdates(Params.LocalUserId, MoveTemp(*Params.Changes))
 	.Then([Promise = MoveTemp(Promise)](TFuture<EOS_EResult>&& Future) mutable
 	{
-		UE_LOG(LogTemp, Warning, TEXT("[FLobbiesEOS::ModifyLobbyMemberDataImpl] Complete. Result: %d"), Future.Get());
+		UE_LOG(LogTemp, Warning, TEXT("[FLobbiesEOS::ModifyLobbyMemberDataImpl] Complete. Result: %s"), *LexToString(Future.Get()));
 
 		// Todo: Handle "no change" better.
 		if (Future.Get() != EOS_EResult::EOS_Success && Future.Get() != EOS_EResult::EOS_NoChange)
diff --git a/Engine/Plugins/Online/Experimental/OnlineServicesEOS/Source/Private/LobbiesEOSTypes.cpp b/Engine/Plugins/Online/Experimental/OnlineServicesEOS/Source/Private/LobbiesEOSTypes.cpp
--- a/Engine/Plugins/Online/Experimental/OnlineServicesEOS/Source/Private/LobbiesEOSTypes.cpp
+++ b/Engine/Plugins/Online/Experimental/OnlineServicesEOS/Source/Private/LobbiesEOSTypes.cpp
@@ -307,6 +307,7 @@ TFuture<EOS_EResult> FLobbyDetailsEOS::ApplyLobbyDataUpdates(FOnlineAccountIdHan
 		// Todo: Schema things.
 
 		EOS_Lobby_AttributeData AttributeData;
+		AttributeData.ApiVersion = EOS_LOBBY_ATTRIBUTEDATA_API_LATEST;
 		AttributeData.Key = KeyConverter.Get();
 		AttributeData.Value.AsUtf8 = ValueConverter.Get();
 		AttributeData.ValueType = EOS_ELobbyAttributeType::EOS_AT_STRING;
diff --git a/Engine/Plugins/Online/Experimental/OnlineServicesEOS/Source/Public/LobbiesEOS.h b/Engine/Plugins/Online/Experimental/OnlineServicesEOS/Source/Public/LobbiesEOS.h
--- a/Engine/Plugins/Online/Experimental/OnlineServicesEOS/Source/Public/LobbiesEOS.h
+++ b/Engine/Plugins/Online/Experimental/OnlineServicesEOS/Source/Public/LobbiesEOS.h
@@ -42,6 +42,7 @@ public:
 	virtual TOnlineAsyncOpHandle<FModifyLobbyJoinPolicy> ModifyLobbyJoinPolicy(FModifyLobbyJoinPolicy::Params&& Params) override;
 	virtual TOnlineAsyncOpHandle<FModifyLobbyAttributes> ModifyLobbyAttributes(FModifyLobbyAttributes::Params&& Params) override;
 	virtual TOnlineAsyncOpHandle<FModifyLobbyMemberAttributes> ModifyLobbyMemberAttributes(FModifyLobbyMemberAttributes::Params&& Params) override;
+	virtual TOnlineResult<FGetJoinedLobbies> GetJoinedLobbies(FGetJoinedLobbies::Params&& Params) override;
 
 private:
 	void HandleLobbyUpdated(const EOS_Lobby_LobbyUpdateReceivedCallbackInfo* Data);
