Implement lobby events.

diff --git a/Engine/Plugins/Online/Experimental/OnlineServices/Source/OnlineServicesCommon/Private/Online/LobbiesCommon.cpp b/Engine/Plugins/Online/Experimental/OnlineServices/Source/OnlineServicesCommon/Private/Online/LobbiesCommon.cpp
--- a/Engine/Plugins/Online/Experimental/OnlineServices/Source/OnlineServicesCommon/Private/Online/LobbiesCommon.cpp
+++ b/Engine/Plugins/Online/Experimental/OnlineServices/Source/OnlineServicesCommon/Private/Online/LobbiesCommon.cpp
@@ -150,7 +150,7 @@ void FLobbiesCommon::RegisterCommands()
 	TOnlineComponent<ILobbies>::RegisterCommands();
 
 	RegisterCommand(&FLobbiesCommon::CreateLobby);
-	RegisterCommand(&FLobbiesCommon::FindLobby);
+	RegisterCommand(&FLobbiesCommon::FindLobbies);
 	RegisterCommand(&FLobbiesCommon::RestoreLobbies);
 	RegisterCommand(&FLobbiesCommon::JoinLobby);
 	RegisterCommand(&FLobbiesCommon::LeaveLobby);
@@ -177,9 +177,9 @@ TOnlineAsyncOpHandle<FCreateLobby> FLobbiesCommon::CreateLobby(FCreateLobby::Par
 	return Operation->GetHandle();
 }
 
-TOnlineAsyncOpHandle<FFindLobby> FLobbiesCommon::FindLobby(FFindLobby::Params&& Params)
+TOnlineAsyncOpHandle<FFindLobbies> FLobbiesCommon::FindLobbies(FFindLobbies::Params&& Params)
 {
-	TOnlineAsyncOpRef<FFindLobby> Operation = GetOp<FFindLobby>(MoveTemp(Params));
+	TOnlineAsyncOpRef<FFindLobbies> Operation = GetOp<FFindLobbies>(MoveTemp(Params));
 	Operation->SetError(Errors::NotImplemented());
 	return Operation->GetHandle();
 }
@@ -273,52 +273,70 @@ TOnlineResult<FGetReceivedInvitations> FLobbiesCommon::GetReceivedInvitations(FG
 
 TOnlineEvent<void(const FLobbyJoined&)> FLobbiesCommon::OnLobbyJoined()
 {
-	return OnLobbyJoinedEvent;
+	return LobbyEvents.OnLobbyJoined;
 }
 
 TOnlineEvent<void(const FLobbyLeft&)> FLobbiesCommon::OnLobbyLeft()
 {
-	return OnLobbyLeftEvent;
+	return LobbyEvents.OnLobbyLeft;
 }
 
 TOnlineEvent<void(const FLobbyMemberJoined&)> FLobbiesCommon::OnLobbyMemberJoined()
 {
-	return OnLobbyMemberJoinedEvent;
+	return LobbyEvents.OnLobbyMemberJoined;
 }
 
 TOnlineEvent<void(const FLobbyMemberLeft&)> FLobbiesCommon::OnLobbyMemberLeft()
 {
-	return OnLobbyMemberLeftEvent;
+	return LobbyEvents.OnLobbyMemberLeft;
 }
 
 TOnlineEvent<void(const FLobbyLeaderChanged&)> FLobbiesCommon::OnLobbyLeaderChanged()
 {
-	return OnLobbyLeaderChangedEvent;
+	return LobbyEvents.OnLobbyLeaderChanged;
 }
 
 TOnlineEvent<void(const FLobbySchemaChanged&)> FLobbiesCommon::OnLobbySchemaChanged()
 {
-	return OnLobbySchemaChangedEvent;
+	return LobbyEvents.OnLobbySchemaChanged;
 }
 
 TOnlineEvent<void(const FLobbyAttributesChanged&)> FLobbiesCommon::OnLobbyAttributesChanged()
 {
-	return OnLobbyAttributesChangedEvent;
+	return LobbyEvents.OnLobbyAttributesChanged;
 }
 
 TOnlineEvent<void(const FLobbyMemberAttributesChanged&)> FLobbiesCommon::OnLobbyMemberAttributesChanged()
 {
-	return OnLobbyMemberAttributesChangedEvent;
+	return LobbyEvents.OnLobbyMemberAttributesChanged;
 }
 
 TOnlineEvent<void(const FLobbyInvitationAdded&)> FLobbiesCommon::OnLobbyInvitationAdded()
 {
-	return OnLobbyInvitationAddedEvent;
+	return LobbyEvents.OnLobbyInvitationAdded;
 }
 
 TOnlineEvent<void(const FLobbyInvitationRemoved&)> FLobbiesCommon::OnLobbyInvitationRemoved()
 {
-	return OnLobbyInvitationRemovedEvent;
+	return LobbyEvents.OnLobbyInvitationRemoved;
+}
+
+// todo: TFuture<void> does not work as expected with "Then".
+TFuture<int> AwaitNextGameTick()
+{
+	TSharedRef<TPromise<int>> Promise = MakeShared<TPromise<int>>();
+	TFuture<int> Future = Promise->GetFuture();
+
+	FTSTicker::GetCoreTicker().AddTicker(
+	TEXT("FLobbiesCommon::AwaitNextGameTick"),
+	0.f,
+	[Promise](float)
+	{
+		Promise->EmplaceValue();
+		return false;
+	});
+
+	return Future;
 }
 
 TFuture<TDefaultErrorResultInternal<FOnlineLobbyIdHandle>> FLobbiesCommon::AwaitInvitation(
@@ -348,7 +366,12 @@ TFuture<TDefaultErrorResultInternal<FOnlineLobbyIdHandle>> FLobbiesCommon::Await
 			{
 				FTSTicker::GetCoreTicker().RemoveTicker(AwaitState->OnAwaitExpiredHandle);
 				AwaitState->bSignaled = true;
-				AwaitState->Promise.EmplaceValue(LobbyId);
+
+				AwaitNextGameTick()
+				.Then([AwaitState, LobbyId](TFuture<int>&&)
+				{
+					AwaitState->Promise.EmplaceValue(LobbyId);
+				});
 			}
 		}
 	});
@@ -517,7 +540,9 @@ TOnlineAsyncOpHandle<FFunctionalTestLobbies> FLobbiesCommon::FunctionalTest(FFun
 	//----------------------------------------------------------------------------------------------
 	// Test 1:
 	//    Step 1: Create a lobby with primary user.
-	//    Step 2: Leave lobby with primary user.
+	//    Step 2: Modify lobby attribute.
+	//    Step 3: Modify lobby member attribute.
+	//    Step 4: Leave lobby with primary user.
 	//----------------------------------------------------------------------------------------------
 
 	.Then([this](TOnlineAsyncOp<FFunctionalTestLobbies>& InAsyncOp)
@@ -534,13 +559,16 @@ TOnlineAsyncOpHandle<FFunctionalTestLobbies> FLobbiesCommon::FunctionalTest(FFun
 		Params.SchemaName = TEXT("test");
 		Params.MaxMembers = 2;
 		Params.JoinPolicy = ELobbyJoinPolicy::InvitationOnly;
-		//Params.Attributes;
-		Params.LocalUsers.Emplace(FJoinLobbyLocalUserData{User1Info.AccountInfo->UserId, {}});
+		Params.Attributes.Add(TEXT("test_attribute_key"), TEXT("test_attribute_value"));
+		Params.LocalUsers.Emplace(FJoinLobbyLocalUserData{User1Info.AccountInfo->UserId, {{TEXT("test_attribute_key"), TEXT("test_attribute_value")}}});
 		return Params;
 	})
 	.Then(CaptureOperationStepResult<FCreateLobby>(*Op, CreateLobbyKeyName, &FLobbiesCommon::CreateLobby))
 	.Then([this](TOnlineAsyncOp<FFunctionalTestLobbies>& InAsyncOp)
 	{
+		const FFunctionalTestLoginUser::Result& User1Info = GetOpDataChecked<FFunctionalTestLoginUser::Result>(InAsyncOp, User1KeyName);
+		const FCreateLobby::Result& CreateResult = GetOpDataChecked<FCreateLobby::Result>(InAsyncOp, CreateLobbyKeyName);
+
 		// Check both lobby joined events were received and in the correct order.
 		const TSharedRef<FLobbyEventCapture>& EventCapture = GetOpDataChecked<TSharedRef<FLobbyEventCapture>>(InAsyncOp, LobbyEventCaptureKeyName);
 		if (EventCapture->GetTotalNotificationsReceived() != 2
@@ -549,6 +577,102 @@ TOnlineAsyncOpHandle<FFunctionalTestLobbies> FLobbiesCommon::FunctionalTest(FFun
 		{
 			InAsyncOp.SetError(Errors::Cancelled());
 		}
+
+		// Check lobby attributes are set to the expected values.
+		if (!CreateResult.Lobby->Attributes.OrderIndependentCompareEqual(
+			TMap<FLobbyAttributeId, FLobbyVariant>{{TEXT("test_attribute_key"), TEXT("test_attribute_value")}}))
+		{
+			InAsyncOp.SetError(Errors::Cancelled());
+		}
+
+		const TSharedRef<const FLobbyMember>* MemberData = CreateResult.Lobby->Members.Find(User1Info.AccountInfo->UserId);
+		if (CreateResult.Lobby->Members.Num() == 1 && MemberData != nullptr)
+		{
+			if (!(**MemberData).Attributes.OrderIndependentCompareEqual(
+				TMap<FLobbyAttributeId, FLobbyVariant>{{TEXT("test_attribute_key"), TEXT("test_attribute_value")}}))
+			{
+				InAsyncOp.SetError(Errors::Cancelled());
+			}
+		}
+		else
+		{
+			InAsyncOp.SetError(Errors::Cancelled());
+		}
+	})
+	.Then([this](TOnlineAsyncOp<FFunctionalTestLobbies>& InAsyncOp)
+	{
+		// Prepare for next test check.
+		GetOpDataChecked<TSharedRef<FLobbyEventCapture>>(InAsyncOp, LobbyEventCaptureKeyName)->Empty();
+
+		const FFunctionalTestLoginUser::Result& User1Info = GetOpDataChecked<FFunctionalTestLoginUser::Result>(InAsyncOp, User1KeyName);
+		const FCreateLobby::Result& CreateLobbyResult = GetOpDataChecked<FCreateLobby::Result>(InAsyncOp, CreateLobbyKeyName);
+
+		FModifyLobbyAttributes::Params Params;
+		Params.LocalUserId = User1Info.AccountInfo->UserId;
+		Params.LobbyId = CreateLobbyResult.Lobby->LobbyId;
+		Params.MutatedAttributes = {{TEXT("test_attribute_key"), TEXT("mutated_test_attribute_value")}};
+		return Params;
+	})
+	.Then(ConsumeOperationStepResult<FModifyLobbyAttributes>(*Op, &FLobbiesCommon::ModifyLobbyAttributes))
+	.Then([this](TOnlineAsyncOp<FFunctionalTestLobbies>& InAsyncOp)
+	{
+		// Check that modification event was received.
+		const TSharedRef<FLobbyEventCapture>& EventCapture = GetOpDataChecked<TSharedRef<FLobbyEventCapture>>(InAsyncOp, LobbyEventCaptureKeyName);
+		if (EventCapture->LobbyAttributesChanged.Num() == 1)
+		{
+			const FLobbyAttributesChanged& Notification = EventCapture->LobbyAttributesChanged[0].Notification;
+			if (!Notification.ChangedAttributes.Difference(TSet<FLobbyAttributeId>{TEXT("test_attribute_key")}).IsEmpty() ||
+				!Notification.Lobby->Attributes.OrderIndependentCompareEqual(TMap<FLobbyAttributeId, FLobbyVariant>{{TEXT("test_attribute_key"), TEXT("mutated_test_attribute_value")}}))
+			{
+				InAsyncOp.SetError(Errors::Cancelled());
+			}
+		}
+		else
+		{
+			InAsyncOp.SetError(Errors::Cancelled());
+		}
+	})
+	.Then([this](TOnlineAsyncOp<FFunctionalTestLobbies>& InAsyncOp)
+	{
+		// Prepare for next test check.
+		GetOpDataChecked<TSharedRef<FLobbyEventCapture>>(InAsyncOp, LobbyEventCaptureKeyName)->Empty();
+
+		const FFunctionalTestLoginUser::Result& User1Info = GetOpDataChecked<FFunctionalTestLoginUser::Result>(InAsyncOp, User1KeyName);
+		const FCreateLobby::Result& CreateLobbyResult = GetOpDataChecked<FCreateLobby::Result>(InAsyncOp, CreateLobbyKeyName);
+
+		FModifyLobbyMemberAttributes::Params Params;
+		Params.LocalUserId = User1Info.AccountInfo->UserId;
+		Params.LobbyId = CreateLobbyResult.Lobby->LobbyId;
+		Params.MutatedAttributes = {{TEXT("test_attribute_key"), TEXT("mutated_test_attribute_value")}};
+		return Params;
+	})
+	.Then(ConsumeOperationStepResult<FModifyLobbyMemberAttributes>(*Op, &FLobbiesCommon::ModifyLobbyMemberAttributes))
+	.Then([this](TOnlineAsyncOp<FFunctionalTestLobbies>& InAsyncOp)
+	{
+		const FFunctionalTestLoginUser::Result& User1Info = GetOpDataChecked<FFunctionalTestLoginUser::Result>(InAsyncOp, User1KeyName);
+
+		// Check that modification event was received.
+		const TSharedRef<FLobbyEventCapture>& EventCapture = GetOpDataChecked<TSharedRef<FLobbyEventCapture>>(InAsyncOp, LobbyEventCaptureKeyName);
+		if (EventCapture->LobbyMemberAttributesChanged.Num() == 1)
+		{
+			const FLobbyMemberAttributesChanged& Notification = EventCapture->LobbyMemberAttributesChanged[0].Notification;
+			if (const TSharedRef<const FLobbyMember>* MemberData = Notification.Lobby->Members.Find(User1Info.AccountInfo->UserId))
+			{
+				if (!Notification.ChangedAttributes.Difference(TSet<FLobbyAttributeId>{TEXT("test_attribute_key")}).IsEmpty() ||
+					!(**MemberData).Attributes.OrderIndependentCompareEqual(TMap<FLobbyAttributeId, FLobbyVariant>{{TEXT("test_attribute_key"), TEXT("mutated_test_attribute_value")}}))
+				{
+					InAsyncOp.SetError(Errors::Cancelled());
+				}
+			}
+			else
+			{
+				InAsyncOp.SetError(Errors::Cancelled());
+			}
+		}
+		else
+		{
+			InAsyncOp.SetError(Errors::Cancelled());
+		}
 	})
 	.Then([this](TOnlineAsyncOp<FFunctionalTestLobbies>& InAsyncOp)
 	{
@@ -661,15 +785,12 @@ TOnlineAsyncOpHandle<FFunctionalTestLobbies> FLobbiesCommon::FunctionalTest(FFun
 	.Then(ConsumeOperationStepResult<FJoinLobby>(*Op, &FLobbiesCommon::JoinLobby))
 	.Then([this](TOnlineAsyncOp<FFunctionalTestLobbies>& InAsyncOp)
 	{
-	// todo: currently busted
-	#if 0
 		// Check both lobby joined events were received and in the correct order.
 		const TSharedRef<FLobbyEventCapture>& EventCapture = GetOpDataChecked<TSharedRef<FLobbyEventCapture>>(InAsyncOp, LobbyEventCaptureKeyName);
 		if (EventCapture->GetTotalNotificationsReceived() != 1 || EventCapture->LobbyMemberJoined.Num() != 1)
 		{
 			InAsyncOp.SetError(Errors::Cancelled());
 		}
-	#endif
 	})
 	.Then([this](TOnlineAsyncOp<FFunctionalTestLobbies>& InAsyncOp)
 	{
@@ -687,15 +808,12 @@ TOnlineAsyncOpHandle<FFunctionalTestLobbies> FLobbiesCommon::FunctionalTest(FFun
 	.Then(ConsumeOperationStepResult<FLeaveLobby>(*Op, &FLobbiesCommon::LeaveLobby))
 	.Then([this](TOnlineAsyncOp<FFunctionalTestLobbies>& InAsyncOp)
 	{
-	// todo: currently busted
-	#if 0
 		// Check for expected events.
 		const TSharedRef<FLobbyEventCapture>& EventCapture = GetOpDataChecked<TSharedRef<FLobbyEventCapture>>(InAsyncOp, LobbyEventCaptureKeyName);
-		if (EventCapture->GetTotalNotificationsReceived() != 1|| EventCapture->LobbyMemberLeft.Num())
+		if (EventCapture->GetTotalNotificationsReceived() != 1 || EventCapture->LobbyMemberLeft.Num() != 1)
 		{
 			InAsyncOp.SetError(Errors::Cancelled());
 		}
-	#endif
 	})
 	.Then([this](TOnlineAsyncOp<FFunctionalTestLobbies>& InAsyncOp)
 	{
@@ -713,8 +831,6 @@ TOnlineAsyncOpHandle<FFunctionalTestLobbies> FLobbiesCommon::FunctionalTest(FFun
 	.Then(ConsumeOperationStepResult<FLeaveLobby>(*Op, &FLobbiesCommon::LeaveLobby))
 	.Then([this](TOnlineAsyncOp<FFunctionalTestLobbies>& InAsyncOp)
 	{
-	// todo: currently busted
-	#if 0
 		// Check for expected events.
 		const TSharedRef<FLobbyEventCapture>& EventCapture = GetOpDataChecked<TSharedRef<FLobbyEventCapture>>(InAsyncOp, LobbyEventCaptureKeyName);
 		if (EventCapture->GetTotalNotificationsReceived() != 2
@@ -723,7 +839,6 @@ TOnlineAsyncOpHandle<FFunctionalTestLobbies> FLobbiesCommon::FunctionalTest(FFun
 		{
 			InAsyncOp.SetError(Errors::Cancelled());
 		}
-	#endif
 	})
 
 	//----------------------------------------------------------------------------------------------
@@ -770,12 +885,12 @@ TOnlineAsyncOpHandle<FFunctionalTestLobbies> FLobbiesCommon::FunctionalTest(FFun
 		const FCreateLobby::Result& CreateLobbyResult = GetOpDataChecked<FCreateLobby::Result>(InAsyncOp, CreateLobbyKeyName);
 
 		// Search for lobby from create
-		FFindLobby::Params Params;
+		FFindLobbies::Params Params;
 		Params.LocalUserId = User2Info.AccountInfo->UserId;
 		Params.LobbyId = CreateLobbyResult.Lobby->LobbyId;
 		return Params;
 	})
-	.Then(CaptureOperationStepResult<FFindLobby>(*Op, FindLobbyKeyName, &FLobbiesCommon::FindLobby))
+	.Then(CaptureOperationStepResult<FFindLobbies>(*Op, FindLobbyKeyName, &FLobbiesCommon::FindLobbies))
 	.Then([this](TOnlineAsyncOp<FFunctionalTestLobbies>& InAsyncOp)
 	{
 		// Prepare for next test check.
@@ -783,9 +898,8 @@ TOnlineAsyncOpHandle<FFunctionalTestLobbies> FLobbiesCommon::FunctionalTest(FFun
 
 		const FFunctionalTestLoginUser::Result& User2Info = GetOpDataChecked<FFunctionalTestLoginUser::Result>(InAsyncOp, User2KeyName);
 		const FCreateLobby::Result& CreateLobbyResult = GetOpDataChecked<FCreateLobby::Result>(InAsyncOp, CreateLobbyKeyName);
-		const FFindLobby::Result& FindResults = GetOpDataChecked<FFindLobby::Result>(InAsyncOp, FindLobbyKeyName);
+		const FFindLobbies::Result& FindResults = GetOpDataChecked<FFindLobbies::Result>(InAsyncOp, FindLobbyKeyName);
 
-		// Not a valid expectation without a custom bucket.
 		if (FindResults.Lobbies.Num() != 1)
 		{
 			InAsyncOp.SetError(Errors::Cancelled());
@@ -802,15 +916,12 @@ TOnlineAsyncOpHandle<FFunctionalTestLobbies> FLobbiesCommon::FunctionalTest(FFun
 	.Then(ConsumeOperationStepResult<FJoinLobby>(*Op, &FLobbiesCommon::JoinLobby))
 	.Then([this](TOnlineAsyncOp<FFunctionalTestLobbies>& InAsyncOp)
 	{
-	// todo: currently busted
-	#if 0
 		// Check both lobby joined events were received and in the correct order.
 		const TSharedRef<FLobbyEventCapture>& EventCapture = GetOpDataChecked<TSharedRef<FLobbyEventCapture>>(InAsyncOp, LobbyEventCaptureKeyName);
 		if (EventCapture->GetTotalNotificationsReceived() != 1 || EventCapture->LobbyMemberJoined.Num() != 1)
 		{
 			InAsyncOp.SetError(Errors::Cancelled());
 		}
-	#endif
 	})
 	.Then([this](TOnlineAsyncOp<FFunctionalTestLobbies>& InAsyncOp)
 	{
@@ -828,15 +939,12 @@ TOnlineAsyncOpHandle<FFunctionalTestLobbies> FLobbiesCommon::FunctionalTest(FFun
 	.Then(ConsumeOperationStepResult<FLeaveLobby>(*Op, &FLobbiesCommon::LeaveLobby))
 	.Then([this](TOnlineAsyncOp<FFunctionalTestLobbies>& InAsyncOp)
 	{
-	// todo: currently busted
-	#if 0
 		// Check for expected events.
 		const TSharedRef<FLobbyEventCapture>& EventCapture = GetOpDataChecked<TSharedRef<FLobbyEventCapture>>(InAsyncOp, LobbyEventCaptureKeyName);
-		if (EventCapture->GetTotalNotificationsReceived() != 1|| EventCapture->LobbyMemberLeft.Num())
+		if (EventCapture->GetTotalNotificationsReceived() != 1|| EventCapture->LobbyMemberLeft.Num() != 1)
 		{
 			InAsyncOp.SetError(Errors::Cancelled());
 		}
-	#endif
 	})
 	.Then([this](TOnlineAsyncOp<FFunctionalTestLobbies>& InAsyncOp)
 	{
@@ -854,8 +962,6 @@ TOnlineAsyncOpHandle<FFunctionalTestLobbies> FLobbiesCommon::FunctionalTest(FFun
 	.Then(ConsumeOperationStepResult<FLeaveLobby>(*Op, &FLobbiesCommon::LeaveLobby))
 	.Then([this](TOnlineAsyncOp<FFunctionalTestLobbies>& InAsyncOp)
 	{
-	// todo: currently busted
-	#if 0
 		// Check for expected events.
 		const TSharedRef<FLobbyEventCapture>& EventCapture = GetOpDataChecked<TSharedRef<FLobbyEventCapture>>(InAsyncOp, LobbyEventCaptureKeyName);
 		if (EventCapture->GetTotalNotificationsReceived() != 2
@@ -864,7 +970,6 @@ TOnlineAsyncOpHandle<FFunctionalTestLobbies> FLobbiesCommon::FunctionalTest(FFun
 		{
 			InAsyncOp.SetError(Errors::Cancelled());
 		}
-	#endif
 	})
 
 	//----------------------------------------------------------------------------------------------
@@ -960,7 +1065,7 @@ TFuture<TOnlineResult<FLobbiesCommon::FFunctionalTestLogoutUser>> FLobbiesCommon
 	LogoutParams.LocalUserId = AccountInfo->UserId;
 
 	AuthInterface->Logout(MoveTemp(LogoutParams))
-	.OnComplete([Promise, PlatformUserId = Params.PlatformUserId](const TOnlineResult<FAuthLogout>& LogoutResult) mutable
+	.OnComplete([Promise, PlatformUserId = Params.PlatformUserId](const TOnlineResult<FAuthLogout>& LogoutResult)
 	{
 		if (LogoutResult.IsError())
 		{
diff --git a/Engine/Plugins/Online/Experimental/OnlineServices/Source/OnlineServicesCommon/Public/Online/LobbiesCommon.h b/Engine/Plugins/Online/Experimental/OnlineServices/Source/OnlineServicesCommon/Public/Online/LobbiesCommon.h
--- a/Engine/Plugins/Online/Experimental/OnlineServices/Source/OnlineServicesCommon/Public/Online/LobbiesCommon.h
+++ b/Engine/Plugins/Online/Experimental/OnlineServices/Source/OnlineServicesCommon/Public/Online/LobbiesCommon.h
@@ -109,6 +109,20 @@ public:
 	TSet<FLobbyAttributeId> ClearedAttributes;
 };
 
+struct FLobbyEvents final
+{
+	TOnlineEventCallable<void(const FLobbyJoined&)> OnLobbyJoined;
+	TOnlineEventCallable<void(const FLobbyLeft&)> OnLobbyLeft;
+	TOnlineEventCallable<void(const FLobbyMemberJoined&)> OnLobbyMemberJoined;
+	TOnlineEventCallable<void(const FLobbyMemberLeft&)> OnLobbyMemberLeft;
+	TOnlineEventCallable<void(const FLobbyLeaderChanged&)> OnLobbyLeaderChanged;
+	TOnlineEventCallable<void(const FLobbySchemaChanged&)> OnLobbySchemaChanged;
+	TOnlineEventCallable<void(const FLobbyAttributesChanged&)> OnLobbyAttributesChanged;
+	TOnlineEventCallable<void(const FLobbyMemberAttributesChanged&)> OnLobbyMemberAttributesChanged;
+	TOnlineEventCallable<void(const FLobbyInvitationAdded&)> OnLobbyInvitationAdded;
+	TOnlineEventCallable<void(const FLobbyInvitationRemoved&)> OnLobbyInvitationRemoved;
+};
+
 class FLobbySchema final
 {
 public:
@@ -185,7 +199,7 @@ public:
 
 	// ILobbies
 	virtual TOnlineAsyncOpHandle<FCreateLobby> CreateLobby(FCreateLobby::Params&& Params) override;
-	virtual TOnlineAsyncOpHandle<FFindLobby> FindLobby(FFindLobby::Params&& Params) override;
+	virtual TOnlineAsyncOpHandle<FFindLobbies> FindLobbies(FFindLobbies::Params&& Params) override;
 	virtual TOnlineAsyncOpHandle<FRestoreLobbies> RestoreLobbies(FRestoreLobbies::Params&& Params) override;
 	virtual TOnlineAsyncOpHandle<FJoinLobby> JoinLobby(FJoinLobby::Params&& Params) override;
 	virtual TOnlineAsyncOpHandle<FLeaveLobby> LeaveLobby(FLeaveLobby::Params&& Params) override;
@@ -505,16 +519,7 @@ protected:
 		};
 	}
 
-	TOnlineEventCallable<void(const FLobbyJoined&)> OnLobbyJoinedEvent;
-	TOnlineEventCallable<void(const FLobbyLeft&)> OnLobbyLeftEvent;
-	TOnlineEventCallable<void(const FLobbyMemberJoined&)> OnLobbyMemberJoinedEvent;
-	TOnlineEventCallable<void(const FLobbyMemberLeft&)> OnLobbyMemberLeftEvent;
-	TOnlineEventCallable<void(const FLobbyLeaderChanged&)> OnLobbyLeaderChangedEvent;
-	TOnlineEventCallable<void(const FLobbySchemaChanged&)> OnLobbySchemaChangedEvent;
-	TOnlineEventCallable<void(const FLobbyAttributesChanged&)> OnLobbyAttributesChangedEvent;
-	TOnlineEventCallable<void(const FLobbyMemberAttributesChanged&)> OnLobbyMemberAttributesChangedEvent;
-	TOnlineEventCallable<void(const FLobbyInvitationAdded&)> OnLobbyInvitationAddedEvent;
-	TOnlineEventCallable<void(const FLobbyInvitationRemoved&)> OnLobbyInvitationRemovedEvent;
+	FLobbyEvents LobbyEvents;
 
 	TSharedPtr<FLobbySchemaRegistry> LobbySchemaRegistry;
 	TSharedPtr<FLobbySchema> ServiceSchema;
diff --git a/Engine/Plugins/Online/Experimental/OnlineServices/Source/OnlineServicesInterface/Private/Lobbies.cpp b/Engine/Plugins/Online/Experimental/OnlineServices/Source/OnlineServicesInterface/Private/Lobbies.cpp
--- a/Engine/Plugins/Online/Experimental/OnlineServices/Source/OnlineServicesInterface/Private/Lobbies.cpp
+++ b/Engine/Plugins/Online/Experimental/OnlineServices/Source/OnlineServicesInterface/Private/Lobbies.cpp
@@ -4,6 +4,117 @@
 
 namespace UE::Online {
 
+const FString LexToString(const FLobbyVariant& Variant)
+{
+	return Variant.GetString();
+}
+
+void LexFromString(FLobbyVariant& Variant, const TCHAR* InStr)
+{
+	Variant.Set(TEXT("")); // todo
+}
+
+int64 FLobbyVariant::GetInt64() const
+{
+	int64 AsInt64 = 0;
+	if (VariantData.IsType<int64>())
+	{
+		AsInt64 = VariantData.Get<int64>();
+	}
+	else if (VariantData.IsType<bool>())
+	{
+		AsInt64 = static_cast<int64>(VariantData.Get<bool>());
+	}
+	else if (VariantData.IsType<double>())
+	{
+		AsInt64 = static_cast<int64>(VariantData.Get<double>());
+	}
+	else if (VariantData.IsType<FString>())
+	{
+		::LexFromString(AsInt64, *VariantData.Get<FString>());
+	}
+	return AsInt64;
+}
+
+double FLobbyVariant::GetDouble() const
+{
+	double AsDouble = 0;
+	if (VariantData.IsType<double>())
+	{
+		AsDouble = VariantData.Get<double>();
+	}
+	else if (VariantData.IsType<FString>())
+	{
+		::LexFromString(AsDouble, *VariantData.Get<FString>());
+	}
+	else if (VariantData.IsType<int64>())
+	{
+		AsDouble = static_cast<double>(VariantData.Get<int64>());
+	}
+	else if (VariantData.IsType<bool>())
+	{
+		AsDouble = static_cast<double>(VariantData.Get<bool>());
+	}
+	return AsDouble;
+}
+
+bool FLobbyVariant::GetBoolean() const
+{
+	bool bAsBool = false;
+	if (VariantData.IsType<bool>())
+	{
+		bAsBool = VariantData.Get<bool>();
+	}
+	else if (VariantData.IsType<FString>())
+	{
+		::LexFromString(bAsBool, *VariantData.Get<FString>());
+	}
+	else
+	{
+		bAsBool = GetInt64() != 0;
+	}
+	return bAsBool;
+}
+
+FString FLobbyVariant::GetString() const
+{
+	if (VariantData.IsType<FString>())
+	{
+		return VariantData.Get<FString>();
+	}
+	else if (VariantData.IsType<int64>())
+	{
+		return FString::Printf(TEXT("%" INT64_FMT), VariantData.Get<int64>());
+	}
+	else if (VariantData.IsType<bool>())
+	{
+		return ::LexToString(VariantData.Get<bool>());
+	}
+	else
+	{
+		checkNoEntry();
+	}
+	return TEXT("");
+}
+
+bool FLobbyVariant::operator==(const FLobbyVariant& Other) const
+{
+	if (VariantData.GetIndex() != Other.VariantData.GetIndex())
+	{
+		return false;
+	}
+
+	switch (VariantData.GetIndex())
+	{
+	case FVariantType::IndexOfType<FString>():	return VariantData.Get<FString>() == Other.VariantData.Get<FString>();
+	case FVariantType::IndexOfType<int64>():	return VariantData.Get<int64>() == Other.VariantData.Get<int64>();
+	case FVariantType::IndexOfType<double>():	return VariantData.Get<double>() == Other.VariantData.Get<double>();
+
+	default:checkNoEntry(); // Intentional fallthrough
+	case FVariantType::IndexOfType<bool>():	return VariantData.Get<bool>() == Other.VariantData.Get<bool>();
+	}
+}
+
 const TCHAR* LexToString(ELobbyJoinPolicy Policy)
 {
 	switch (Policy)
diff --git a/Engine/Plugins/Online/Experimental/OnlineServices/Source/OnlineServicesInterface/Public/Online/Lobbies.h b/Engine/Plugins/Online/Experimental/OnlineServices/Source/OnlineServicesInterface/Public/Online/Lobbies.h
--- a/Engine/Plugins/Online/Experimental/OnlineServices/Source/OnlineServicesInterface/Public/Online/Lobbies.h
+++ b/Engine/Plugins/Online/Experimental/OnlineServices/Source/OnlineServicesInterface/Public/Online/Lobbies.h
@@ -5,30 +5,79 @@
 #include "Online/OnlineAsyncOpHandle.h"
 #include "Online/CoreOnline.h"
 #include "Online/OnlineMeta.h"
+#include "Misc/TVariant.h"
 
 namespace UE::Online {
 
-typedef FName FLobbySchemaId;
-typedef FName FLobbyAttributeId;
-using FLobbyVariant = FString; // TODO:  Import FVariantData
+using FLobbySchemaId = FName;
+using FLobbyAttributeId = FName;
+
+class FLobbyVariant
+{
+public:
+	using FVariantType = TVariant<FString, int64, double, bool>;
+	FLobbyVariant() = default;
+	FLobbyVariant(const FLobbyVariant& InOther) : VariantData(InOther.VariantData) {}
+	FLobbyVariant(FLobbyVariant&& InOther) : VariantData(MoveTemp(InOther.VariantData)) {}
+	FLobbyVariant& operator=(FLobbyVariant&&) = default;
+	FLobbyVariant& operator=(const FLobbyVariant&) = default;
+
+	template<typename ValueType>
+	FLobbyVariant(const ValueType& InData) { Set(MoveTempIfPossible(InData)); }
+	void Set(const TCHAR* AsString) { VariantData.Emplace<FString>(AsString); }
+	void Set(const FString& AsString) { VariantData.Emplace<FString>(AsString); }
+	void Set(FString&& AsString) { VariantData.Emplace<FString>(MoveTemp(AsString)); }
+	void Set(int64 AsInt) { VariantData.Emplace<int64>(AsInt); }
+	void Set(double AsDouble) { VariantData.Emplace<double>(AsDouble); }
+	void Set(bool bAsBool) { VariantData.Emplace<bool>(bAsBool); }
+	ONLINESERVICESINTERFACE_API int64 GetInt64() const;
+	ONLINESERVICESINTERFACE_API double GetDouble() const;
+	ONLINESERVICESINTERFACE_API bool GetBoolean() const;
+	ONLINESERVICESINTERFACE_API FString GetString() const;
+
+	ONLINESERVICESINTERFACE_API bool operator==(const FLobbyVariant& Other) const;
+	bool operator!=(const FLobbyVariant& Other) const { return !(*this == Other); }
+public:
+	FVariantType VariantData;
+};
+ONLINESERVICESINTERFACE_API const FString LexToString(const FLobbyVariant& Variant);
+ONLINESERVICESINTERFACE_API void LexFromString(FLobbyVariant& Variant, const TCHAR* InStr);
 
 enum class ELobbyJoinPolicy : uint8
 {
 	/** 
 	* Lobby can be found through searches based on attribute matching,
-	*  by knowing the lobby id, or by invitation.
+	* by knowing the lobby id, or by invitation.
 	*/
 	PublicAdvertised,
 
-	/**  Lobby may be joined by knowing the lobby id or by invitation. */
+	/** Lobby may be joined by knowing the lobby id or by invitation. */
 	PublicNotAdvertised,
 
-	/**  Lobby may only be joined by invitation. */
+	/** Lobby may only be joined by invitation. */
 	InvitationOnly,
 };
 ONLINESERVICESINTERFACE_API const TCHAR* LexToString(ELobbyJoinPolicy Policy);
 ONLINESERVICESINTERFACE_API void LexFromString(ELobbyJoinPolicy& OutPolicy, const TCHAR* InStr);
 
+enum class ELobbyMemberLeaveReason
+{
+	/** The lobby member explicitly left the lobby. */
+	Left,
+
+	/** The lobby member was kicked from the lobby by the lobby owner. */
+	Kicked,
+
+	/** The lobby member unexpectedly left. */
+	Disconnected,
+
+	/**
+	* The lobby was destroyed by the service.
+	* All members have left.
+	*/
+	Closed
+};
+
 enum class ELobbyAttributeVisibility
 {
 	Public,
@@ -52,7 +101,7 @@ struct FLobby
 	int32 MaxMembers;
 	ELobbyJoinPolicy JoinPolicy;
 	TMap<FLobbyAttributeId, FLobbyVariant> Attributes;
-	TArray<TSharedRef<const FLobbyMember>> Members;
+	TMap<FOnlineAccountIdHandle, TSharedRef<const FLobbyMember>> Members;
 };
 
 struct FJoinLobbyLocalUserData
@@ -132,33 +181,33 @@ struct FCreateLobby
 	};
 };
 
-struct FFindLobby
+struct FFindLobbies
 {
-	static constexpr TCHAR Name[] = TEXT("FindLobby");
+	static constexpr TCHAR Name[] = TEXT("FindLobbies");
 
-	/** Input struct for Lobbies::FindLobby */
+	/** Input struct for Lobbies::FindLobbies */
 	struct Params
 	{
 		/** The local user agent which will perform the action. */
 		FOnlineAccountIdHandle LocalUserId;
 
-		/** 
+		/**
 		* Max results to return in one search.
 		*  Actual count may be smaller based on implementation.
 		*/
 		uint32 MaxResults = 20;
 
-		/**  Filters to apply when searching for lobbies. */
+		/** Filters to apply when searching for lobbies. */
 		TArray<FFindLobbySearchFilter> Filters;
 
-		/**  Find lobbies containing the target user. */
+		/** Find lobbies containing the target user. */
 		TOptional<FOnlineAccountIdHandle> TargetUser;
 
-		/**  Find join info for the target lobby id. */
+		/** Find join info for the target lobby id. */
 		TOptional<FOnlineLobbyIdHandle> LobbyId;
 	};
 
-	/** Output struct for Lobbies::FindLobby */
+	/** Output struct for Lobbies::FindLobbies */
 	struct Result
 	{
 		TArray<TSharedRef<const FLobby>> Lobbies;
@@ -187,16 +236,16 @@ struct FJoinLobby
 	/** Input struct for Lobbies::JoinLobby */
 	struct Params
 	{
-		/**  The local user agent which will perform the action. */
+		/** The local user agent which will perform the action. */
 		FOnlineAccountIdHandle LocalUserId;
 
-		/**  The local name for the lobby. */
+		/** The local name for the lobby. */
 		FName LocalName;
 
-		/**  The id of the lobby to be joined. */
+		/** The id of the lobby to be joined. */
 		FOnlineLobbyIdHandle LobbyId;
 
-		/**  Local users who will be joining the lobby. */
+		/** Local users who will be joining the lobby. */
 		TArray<FJoinLobbyLocalUserData> LocalUsers;
 	};
 
@@ -215,10 +264,10 @@ struct FLeaveLobby
 	/** Input struct for Lobbies::LeaveLobby */
 	struct Params
 	{
-		/**  The local user agent which will perform the action. */
+		/** The local user agent which will perform the action. */
 		FOnlineAccountIdHandle LocalUserId;
 
-		/**  Id of the lobby to leave. */
+		/** Id of the lobby to leave. */
 		FOnlineLobbyIdHandle LobbyId;
 	};
 
@@ -235,13 +284,13 @@ struct FInviteLobbyMember
 	/** Input struct for Lobbies::InviteLobbyMember */
 	struct Params
 	{
-		/**  The local user agent which will perform the action. */
+		/** The local user agent which will perform the action. */
 		FOnlineAccountIdHandle LocalUserId;
 
-		/**  Id of the lobby for which the invitation will be sent. */
+		/** Id of the lobby for which the invitation will be sent. */
 		FOnlineLobbyIdHandle LobbyId;
 
-		/**  Id of the player who will be sent the invitation. */
+		/** Id of the player who will be sent the invitation. */
 		FOnlineAccountIdHandle TargetUserId;
 	};
 
@@ -258,10 +307,10 @@ struct FDeclineLobbyInvitation
 	/** Input struct for Lobbies::DeclineLobbyInvitation */
 	struct Params
 	{
-		/**  The local user agent which will perform the action. */
+		/** The local user agent which will perform the action. */
 		FOnlineAccountIdHandle LocalUserId;
 
-		// Id of the lobby for which the invitations will be declined.
+		/** Id of the lobby for which the invitations will be declined. */
 		FOnlineLobbyIdHandle LobbyId;
 	};
 
@@ -421,7 +470,7 @@ struct FGetJoinedLobbies
 	/** Output struct for Lobbies::GetJoinedLobbies */
 	struct Result
 	{
-		TArray<TSharedRef<FLobby>> Lobbies;
+		TArray<TSharedRef<const FLobby>> Lobbies;
 	};
 };
 
@@ -478,26 +527,50 @@ struct FLobbyMemberLeft
 	/** Member data access. */
 	TSharedRef<const FLobbyMember> Member;
 
+	/** Context for the member leaving. */
+	ELobbyMemberLeaveReason Reason;
+
 };
 
 /** Struct for LobbyLeaderChanged event */
 struct FLobbyLeaderChanged
 {
+	/** Lobby data access. */
+	TSharedRef<const FLobby> Lobby;
+
+	/** Leader data access. */
+	TSharedRef<const FLobbyMember> Leader;
 };
 
 /** Struct for LobbySchemaChanged event */
 struct FLobbySchemaChanged
 {
+	/** Lobby data access. */
+	TSharedRef<const FLobby> Lobby;
 };
 
 /** Struct for LobbyAttributesChanged event */
 struct FLobbyAttributesChanged
 {
+	/** Lobby data access. */
+	TSharedRef<const FLobby> Lobby;
+
+	/** Attribute keys which changed. */
+	TSet<FLobbyAttributeId> ChangedAttributes;
+
 };
 
 /** Struct for LobbyMemberAttributesChanged event */
 struct FLobbyMemberAttributesChanged
 {
+	/** Lobby data access. */
+	TSharedRef<const FLobby> Lobby;
+
+	/** Member data access. */
+	TSharedRef<const FLobbyMember> Member;
+
+	/** Attribute keys which changed. */
+	TSet<FLobbyAttributeId> ChangedAttributes;
 };
 
 /** Struct for LobbyInvitationAdded event */
@@ -547,7 +620,7 @@ public:
 	 * @param Params for the FindLobby call
 	 * @return
 	 */
-	virtual TOnlineAsyncOpHandle<FFindLobby> FindLobby(FFindLobby::Params&& Params) = 0;
+	virtual TOnlineAsyncOpHandle<FFindLobbies> FindLobbies(FFindLobbies::Params&& Params) = 0;
 
 	/**
 	 * Try to rejoin previously joined lobbies.
@@ -774,10 +847,10 @@ BEGIN_ONLINE_STRUCT_META(FCreateLobby::Result)
 	ONLINE_STRUCT_FIELD(FCreateLobby::Result, Lobby)
 END_ONLINE_STRUCT_META()
 
-BEGIN_ONLINE_STRUCT_META(FFindLobby::Params)
+BEGIN_ONLINE_STRUCT_META(FFindLobbies::Params)
 END_ONLINE_STRUCT_META()
 
-BEGIN_ONLINE_STRUCT_META(FFindLobby::Result)
+BEGIN_ONLINE_STRUCT_META(FFindLobbies::Result)
 END_ONLINE_STRUCT_META()
 
 BEGIN_ONLINE_STRUCT_META(FRestoreLobbies::Params)
diff --git a/Engine/Plugins/Online/Experimental/OnlineServicesEOS/Source/Private/LobbiesEOS.cpp b/Engine/Plugins/Online/Experimental/OnlineServicesEOS/Source/Private/LobbiesEOS.cpp
--- a/Engine/Plugins/Online/Experimental/OnlineServicesEOS/Source/Private/LobbiesEOS.cpp
+++ b/Engine/Plugins/Online/Experimental/OnlineServicesEOS/Source/Private/LobbiesEOS.cpp
@@ -14,7 +14,9 @@
 namespace UE::Online {
 
 static const FString LobbyDataKeyName = TEXT("LobbyData");
-static const FString LobbyNotificationPauseKeyName = TEXT("LobbyNotificationPause");
+static const FString LobbyDetailsKeyName = TEXT("LobbyDetails");
+static const FString LobbyChangesKeyName = TEXT("LobbyChanges");
+static const FString LobbyMemberChangesKeyName = TEXT("LobbyMemberChanges");
 
 static const int32 MaxAttributeSize = 1000;
 
@@ -45,7 +47,7 @@ void FLobbiesEOS::Initialize()
 
 	LobbyPrerequisites = MakeShared<FLobbyPrerequisitesEOS>(FLobbyPrerequisitesEOS{
 		LobbyInterfaceHandle,
-		StaticCastSharedPtr<FAuthEOS>(Services.GetAuthInterface()).ToSharedRef(),
+		StaticCastSharedPtr<FAuthEOS>(Services.GetAuthInterface()),
 		LobbySchemaRegistry.ToSharedRef(),
 		ServiceSchema.ToSharedRef()});
 
@@ -104,7 +106,9 @@ TOnlineAsyncOpHandle<FCreateLobby> FLobbiesEOS::CreateLobby(FCreateLobby::Params
 
 	Op->Then([this](TOnlineAsyncOp<FCreateLobby>& InAsyncOp)
 	{
-		const auto& Params = InAsyncOp.GetParams();
+		// Step 1: Call create lobby.
+
+		const FCreateLobby::Params& Params = InAsyncOp.GetParams();
 		const FTCHARToUTF8 BucketIdUTF8(*LobbyBucketId);
 
 		// The lobby will be created as invitation only.
@@ -121,66 +125,100 @@ TOnlineAsyncOpHandle<FCreateLobby> FLobbiesEOS::CreateLobby(FCreateLobby::Params
 		CreateLobbyOptions.bDisableHostMigration = false; // todo: handle
 		CreateLobbyOptions.bEnableRTCRoom = false; // todo: handle
 
-		UE_LOG(LogTemp, Warning, TEXT("[FLobbiesEOS::EOS_Lobby_CreateLobby] Start. Member: %s"), *ToLogString(Params.LocalUserId));
-
 		return EOS_Async<EOS_Lobby_CreateLobbyCallbackInfo>(EOS_Lobby_CreateLobby, LobbyPrerequisites->LobbyInterfaceHandle, CreateLobbyOptions); 
 	})
 	.Then([this](TOnlineAsyncOp<FCreateLobby>& InAsyncOp, const EOS_Lobby_CreateLobbyCallbackInfo* Data)
 	{
-		UE_LOG(LogTemp, Warning, TEXT("[FLobbiesEOS::EOS_Lobby_CreateLobby] Complete. Result: %d"), Data->ResultCode);
+		// Step 2: Start creating the lobby data from the EOS lobby details object.
+
+		const FCreateLobby::Params& Params = InAsyncOp.GetParams();
 
 		if (Data->ResultCode != EOS_EResult::EOS_Success)
 		{
 			// Todo: errors
 			InAsyncOp.SetError(Errors::Unknown(FromEOSError(Data->ResultCode)));
-			return MakeFulfilledPromise<TDefaultErrorResultInternal<TSharedPtr<FLobbyDataEOS>>>().GetFuture();
+			return MakeFulfilledPromise<TDefaultErrorResultInternal<TSharedRef<FLobbyDataEOS>>>().GetFuture();
+		}
+
+		TDefaultErrorResultInternal<TSharedRef<FLobbyDetailsEOS>> LobbyDetailsResult =
+			FLobbyDetailsEOS::CreateFromLobbyId(LobbyPrerequisites.ToSharedRef(), Params.LocalUserId, Data->LobbyId);
+		if (LobbyDetailsResult.IsError())
+		{
+			// Todo: manually call eos destroy lobby here.
+
+			// Todo: errors
+			InAsyncOp.SetError(MoveTemp(LobbyDetailsResult.GetErrorValue()));
+			return MakeFulfilledPromise<TDefaultErrorResultInternal<TSharedRef<FLobbyDataEOS>>>().GetFuture();
 		}
 		
-		return LobbyDataRegistry->FindOrCreateFromLobbyId(InAsyncOp.GetParams().LocalUserId, Data->LobbyId);
+		return LobbyDataRegistry->FindOrCreateFromLobbyDetails(InAsyncOp.GetParams().LocalUserId, LobbyDetailsResult.GetOkValue());
 	})
-	.Then([this, DestroyLobbyDuringCreate](TOnlineAsyncOp<FCreateLobby>& InAsyncOp, TDefaultErrorResultInternal<TSharedPtr<FLobbyDataEOS>>&& Result)
+	.Then([this, DestroyLobbyDuringCreate](TOnlineAsyncOp<FCreateLobby>& InAsyncOp, TDefaultErrorResultInternal<TSharedRef<FLobbyDataEOS>>&& Result)
 	{
-		const auto& Params = InAsyncOp.GetParams();
+		// Step 3: Store the lobby data on the async op properties.
+
+		const FCreateLobby::Params& Params = InAsyncOp.GetParams();
 
 		if (Result.IsError())
 		{
 			// Todo: destroy lobby.
+			// Change DestroyLobbyImpl to take an EOS_LobbyId instead of an FLobbyDataEOS.
 
 			// Todo: errors
 			InAsyncOp.SetError(Errors::Unknown(MoveTemp(Result.GetErrorValue())));
 			return;
 		}
 
-		TSharedRef<FLobbyDataEOS> LobbyData = Result.GetOkValue().ToSharedRef();
-
-		// Pause notifications for all joining local users.
-		TArray<TSharedRef<FLobbyNotificationPauseHandle>> PauseHandles;
-		for (const FJoinLobbyLocalUserData& UserData : Params.LocalUsers)
-		{
-			PauseHandles.Add(LobbyData->PauseLobbyNotifications(UserData.LocalUserId));
-		}
-
+		// Store lobby data on the operation.
+		TSharedRef<FLobbyDataEOS> LobbyData = Result.GetOkValue();
 		InAsyncOp.Data.Set<TSharedRef<FLobbyDataEOS>>(LobbyDataKeyName, MoveTemp(LobbyData));
-		InAsyncOp.Data.Set<TArray<TSharedRef<FLobbyNotificationPauseHandle>>>(LobbyNotificationPauseKeyName, MoveTemp(PauseHandles));
 	})
 	.Then([this](TOnlineAsyncOp<FCreateLobby>& InAsyncOp)
 	{
-		const auto& Params = InAsyncOp.GetParams();
-		auto LobbyData = GetOpDataChecked<TSharedRef<FLobbyDataEOS>>(InAsyncOp, LobbyDataKeyName);
+		// Step 4: Set attributes for the lobby creator.
 
-		// Add creator attributes.
-		FModifyLobbyMemberDataImpl::Params ModifyLobbyMemberDataParams;
-		ModifyLobbyMemberDataParams.LobbyData = LobbyData;
-		ModifyLobbyMemberDataParams.LocalUserId = Params.LocalUserId;
-		ModifyLobbyMemberDataParams.Changes = MakeShared<FLobbyMemberDataChanges>();
-		//ModifyLobbyMemberDataParams.Changes->MutatedAttributes;
+		const FCreateLobby::Params& Params = InAsyncOp.GetParams();
+		const TSharedRef<FLobbyDataEOS>& LobbyData = GetOpDataChecked<TSharedRef<FLobbyDataEOS>>(InAsyncOp, LobbyDataKeyName);
 
-		return ModifyLobbyMemberDataImpl(MoveTemp(ModifyLobbyMemberDataParams));
+		// Add creator attributes if set.
+		// Todo: make this nicer.
+		const FJoinLobbyLocalUserData* UserData = Params.LocalUsers.FindByPredicate(
+		[LocalUserId = Params.LocalUserId](const FJoinLobbyLocalUserData& Data)
+		{
+			return LocalUserId == Data.LocalUserId;
+		});
+
+		TSharedRef<FClientLobbyMemberDataChanges> LobbyOwnerAttributes = MakeShared<FClientLobbyMemberDataChanges>();
+
+		// Add owner attributes to operation data.
+		// CreatingMemberData is used to update the local lobby data and for dispatching notifications once creation has completed.
+		TMap<FOnlineAccountIdHandle, TSharedRef<FClientLobbyMemberDataChanges>> CreatingMemberData;
+		CreatingMemberData.Reserve(Params.LocalUsers.Num());
+		CreatingMemberData.Add(Params.LocalUserId, LobbyOwnerAttributes);
+		InAsyncOp.Data.Set<TMap<FOnlineAccountIdHandle, TSharedRef<FClientLobbyMemberDataChanges>>>(LobbyMemberChangesKeyName, MoveTemp(CreatingMemberData));
+
+		if (UserData)
+		{
+			LobbyOwnerAttributes->MutatedAttributes = UserData->Attributes;
+
+			FModifyLobbyMemberDataImpl::Params ModifyLobbyMemberDataParams;
+			ModifyLobbyMemberDataParams.LobbyData = LobbyData;
+			ModifyLobbyMemberDataParams.LocalUserId = Params.LocalUserId;
+			ModifyLobbyMemberDataParams.Changes = LobbyOwnerAttributes;
+
+			return ModifyLobbyMemberDataImpl(MoveTemp(ModifyLobbyMemberDataParams));
+		}
+		else
+		{
+			return MakeFulfilledPromise<TDefaultErrorResult<FModifyLobbyMemberDataImpl>>().GetFuture();
+		}
 	})
 	.Then([this, DestroyLobbyDuringCreate](TOnlineAsyncOp<FCreateLobby>& InAsyncOp, TDefaultErrorResult<FModifyLobbyMemberDataImpl>&& ModifyLobbyOwnerResult)
 	{
-		const auto& Params = InAsyncOp.GetParams();
-		auto LobbyData = GetOpDataChecked<TSharedRef<FLobbyDataEOS>>(InAsyncOp, LobbyDataKeyName);
+		// Step 5: Handle result
+
+		const FCreateLobby::Params& Params = InAsyncOp.GetParams();
+		TSharedRef<FLobbyDataEOS> LobbyData = GetOpDataChecked<TSharedRef<FLobbyDataEOS>>(InAsyncOp, LobbyDataKeyName);
 
 		if (ModifyLobbyOwnerResult.IsError())
 		{
@@ -193,29 +231,38 @@ TOnlineAsyncOpHandle<FCreateLobby> FLobbiesEOS::CreateLobby(FCreateLobby::Params
 	})
 	.Then([this](TOnlineAsyncOp<FCreateLobby>& InAsyncOp)
 	{
-		// Todo: Add other local members.
+		// Step 6: Todo: Add other local members.
+
+		// Store member attributes on the operation.
 	})
 	.Then([this](TOnlineAsyncOp<FCreateLobby>& InAsyncOp)
 	{
-		const auto& Params = InAsyncOp.GetParams();
-		auto LobbyData = GetOpDataChecked<TSharedRef<FLobbyDataEOS>>(InAsyncOp, LobbyDataKeyName);
+		// Step 7: Add lobby attributes, set lobby join policy to user setting.
+
+		const FCreateLobby::Params& Params = InAsyncOp.GetParams();
+		const TSharedRef<FLobbyDataEOS>& LobbyData = GetOpDataChecked<TSharedRef<FLobbyDataEOS>>(InAsyncOp, LobbyDataKeyName);
+
+		// Track lobby changes to sync client-side cache at operation completion.
+		TSharedRef<FClientLobbyDataChanges> LobbyChanges = MakeShared<FClientLobbyDataChanges>();
+		InAsyncOp.Data.Set<TSharedRef<FClientLobbyDataChanges>>(LobbyChangesKeyName, LobbyChanges);
+		LobbyChanges->MutatedAttributes = Params.Attributes;
+		LobbyChanges->JoinPolicy = Params.JoinPolicy;
 
 		// Add lobby attributes.
 		// Set lobby privacy setting to the user provided value.
-
 		FModifyLobbyDataImpl::Params ModifyLobbyDataParams;
 		ModifyLobbyDataParams.LobbyData = LobbyData;
 		ModifyLobbyDataParams.LocalUserId = Params.LocalUserId;
-		ModifyLobbyDataParams.Changes = MakeShared<FLobbyDataChanges>();
-		ModifyLobbyDataParams.Changes->MutatedAttributes = Params.Attributes;
-		ModifyLobbyDataParams.Changes->JoinPolicy = Params.JoinPolicy;
+		ModifyLobbyDataParams.Changes = LobbyChanges;
 
 		return ModifyLobbyDataImpl(MoveTemp(ModifyLobbyDataParams));
 	})
 	.Then([this, DestroyLobbyDuringCreate](TOnlineAsyncOp<FCreateLobby>& InAsyncOp, TDefaultErrorResult<FModifyLobbyDataImpl>&& ModifyLobbyResult) mutable
 	{
-		const auto& Params = InAsyncOp.GetParams();
-		auto LobbyData = GetOpDataChecked<TSharedRef<FLobbyDataEOS>>(InAsyncOp, LobbyDataKeyName);
+		// Step 8: Handle result.
+
+		const FCreateLobby::Params& Params = InAsyncOp.GetParams();
+		const TSharedRef<FLobbyDataEOS>& LobbyData = GetOpDataChecked<TSharedRef<FLobbyDataEOS>>(InAsyncOp, LobbyDataKeyName);
 
 		if (ModifyLobbyResult.IsError())
 		{
@@ -228,32 +275,40 @@ TOnlineAsyncOpHandle<FCreateLobby> FLobbiesEOS::CreateLobby(FCreateLobby::Params
 	})
 	.Then([this](TOnlineAsyncOp<FCreateLobby>& InAsyncOp)
 	{
-		const auto& Params = InAsyncOp.GetParams();
-		auto LobbyData = GetOpDataChecked<TSharedRef<FLobbyDataEOS>>(InAsyncOp, LobbyDataKeyName);
-		LobbyData->GetLobbyImpl()->LocalName = Params.LocalName;
+		// Step 9: Add the lobby to the active list for each member, apply changes to the cached
+		// lobby object, and signal notifications.
+		// The active lobbies list holds a reference to the lobby data to keep it from being cleaned up.
 
-		// Remove notification pause for joining local users.
-		TArray<TSharedRef<FLobbyNotificationPauseHandle>> ClearedPauses;
-		InAsyncOp.Data.Set<TArray<TSharedRef<FLobbyNotificationPauseHandle>>>(LobbyNotificationPauseKeyName, MoveTemp(ClearedPauses));
+		const FCreateLobby::Params& Params = InAsyncOp.GetParams();
+		const TSharedRef<FLobbyDataEOS>& LobbyData = GetOpDataChecked<TSharedRef<FLobbyDataEOS>>(InAsyncOp, LobbyDataKeyName);
 
 		// Mark the lobby active for each member.
-		// Lobby activation / deactivation will dispatch lobby join / leave notifications.
 		for (const FJoinLobbyLocalUserData& LocalUser : Params.LocalUsers)
 		{
 			AddActiveLobby(LocalUser.LocalUserId, LobbyData);
 		}
 
-		InAsyncOp.SetResult(FCreateLobby::Result{LobbyData->GetLobbyImpl()});
+		// Add member changes to the lobby changes object.
+		const TSharedRef<FClientLobbyDataChanges>& LobbyChanges = GetOpDataChecked<TSharedRef<FClientLobbyDataChanges>>(InAsyncOp, LobbyChangesKeyName);
+		const TMap<FOnlineAccountIdHandle, TSharedRef<FClientLobbyMemberDataChanges>>& MemberChanges =
+			GetOpDataChecked<TMap<FOnlineAccountIdHandle, TSharedRef<FClientLobbyMemberDataChanges>>>(InAsyncOp, LobbyMemberChangesKeyName);
+		LobbyChanges->MutatedMembers = MemberChanges;
+		LobbyChanges->LocalName = InAsyncOp.GetParams().LocalName;
+
+		// Make local changes to lobby data and generate notifications.
+		LobbyData->GetClientLobbyData()->ApplyLobbyUpdateFromLocalChanges(MoveTemp(*LobbyChanges), LobbyEvents);
+
+		InAsyncOp.SetResult(FCreateLobby::Result{LobbyData->GetClientLobbyData()->GetPublicDataPtr()});
 	})
-	.Enqueue(GetSerialQueue(Params.LocalUserId));
+	.Enqueue(GetSerialQueue());
 
 	return Op->GetHandle();
 }
 
-TOnlineAsyncOpHandle<FFindLobby> FLobbiesEOS::FindLobby(FFindLobby::Params&& InParams)
+TOnlineAsyncOpHandle<FFindLobbies> FLobbiesEOS::FindLobbies(FFindLobbies::Params&& InParams)
 {
-	TOnlineAsyncOpRef<FFindLobby> Op = GetOp<FFindLobby>(MoveTemp(InParams));
-	const FFindLobby::Params& Params = Op->GetParams();
+	TOnlineAsyncOpRef<FFindLobbies> Op = GetOp<FFindLobbies>(MoveTemp(InParams));
+	const FFindLobbies::Params& Params = Op->GetParams();
 
 	if (!Services.Get<FAuthEOS>()->IsLoggedIn(Params.LocalUserId))
 	{
@@ -264,11 +319,11 @@ TOnlineAsyncOpHandle<FFindLobby> FLobbiesEOS::FindLobby(FFindLobby::Params&& InP
 	// Invalidate previous search results.
 	ActiveSearchResults.Remove(Params.LocalUserId);
 
-	Op->Then([this](TOnlineAsyncOp<FFindLobby>& InAsyncOp)
+	Op->Then([this](TOnlineAsyncOp<FFindLobbies>& InAsyncOp)
 	{
 		return FLobbySearchEOS::Create(LobbyPrerequisites.ToSharedRef(), LobbyDataRegistry.ToSharedRef(), InAsyncOp.GetParams());
 	})
-	.Then([this](TOnlineAsyncOp<FFindLobby>& InAsyncOp, TDefaultErrorResultInternal<TSharedPtr<FLobbySearchEOS>> Result)
+	.Then([this](TOnlineAsyncOp<FFindLobbies>& InAsyncOp, TDefaultErrorResultInternal<TSharedRef<FLobbySearchEOS>> Result)
 	{
 		if (Result.IsError())
 		{
@@ -277,11 +332,11 @@ TOnlineAsyncOpHandle<FFindLobby> FLobbiesEOS::FindLobby(FFindLobby::Params&& InP
 		}
 		else
 		{
-			ActiveSearchResults.Add(InAsyncOp.GetParams().LocalUserId, Result.GetOkValue().ToSharedRef());
-			InAsyncOp.SetResult(FFindLobby::Result{Result.GetOkValue()->GetLobbyResults()});
+			ActiveSearchResults.Add(InAsyncOp.GetParams().LocalUserId, Result.GetOkValue());
+			InAsyncOp.SetResult(FFindLobbies::Result{Result.GetOkValue()->GetLobbyResults()});
 		}
 	})
-	.Enqueue(GetSerialQueue(Params.LocalUserId));
+	.Enqueue(GetSerialQueue());
 
 	return Op->GetHandle();
 }
@@ -306,15 +361,9 @@ TOnlineAsyncOpHandle<FJoinLobby> FLobbiesEOS::JoinLobby(FJoinLobby::Params&& InP
 
 	Op->Then([this, LobbyData](TOnlineAsyncOp<FJoinLobby>& InAsyncOp)
 	{
-		const FJoinLobby::Params& Params = InAsyncOp.GetParams();
+		// Join all users to the lobby.
 
-		// Pause notifications for all joining local users.
-		TArray<TSharedRef<FLobbyNotificationPauseHandle>> PauseHandles;
-		for (const FJoinLobbyLocalUserData& UserData : Params.LocalUsers)
-		{
-			PauseHandles.Add(LobbyData->PauseLobbyNotifications(UserData.LocalUserId));
-		}
-		InAsyncOp.Data.Set<TArray<TSharedRef<FLobbyNotificationPauseHandle>>>(LobbyNotificationPauseKeyName, MoveTemp(PauseHandles));
+		const FJoinLobby::Params& Params = InAsyncOp.GetParams();
 
 		FJoinLobbyImpl::Params JoinParams;
 		JoinParams.LobbyData = LobbyData;
@@ -323,31 +372,41 @@ TOnlineAsyncOpHandle<FJoinLobby> FLobbiesEOS::JoinLobby(FJoinLobby::Params&& InP
 		JoinParams.LocalUsers = Params.LocalUsers;
 		return JoinLobbyImpl(MoveTemp(JoinParams));
 	})
-	.Then([this, LobbyData](TOnlineAsyncOp<FJoinLobby>& InAsyncOp, TDefaultErrorResult<FJoinLobbyImpl> Result)
+	.Then([this](TOnlineAsyncOp<FJoinLobby>& InAsyncOp, TDefaultErrorResult<FJoinLobbyImpl> Result)
 	{
-		const auto& Params = InAsyncOp.GetParams();
-
+		// Handle result.
 		if (Result.IsError())
 		{
 			// TODO: Error codes
 			InAsyncOp.SetError(Errors::Unknown(Result.GetErrorValue()));
 		}
-		else
+	})
+	.Then([this, LobbyData](TOnlineAsyncOp<FJoinLobby>& InAsyncOp)
+	{
+		// Mark the lobby active for each member.
+		// Add users to local lobby data and dispatch notifications.
+
+		const FJoinLobby::Params& Params = InAsyncOp.GetParams();
+
+		// Mark the lobby active for each member.
+		for (const FJoinLobbyLocalUserData& LocalUser : Params.LocalUsers)
 		{
-			LobbyData->GetLobbyImpl()->LocalName = InAsyncOp.GetParams().LocalName;
-			// Remove notification pause for joining local users.
-			TArray<TSharedRef<FLobbyNotificationPauseHandle>> ClearedPauses;
-			InAsyncOp.Data.Set<TArray<TSharedRef<FLobbyNotificationPauseHandle>>>(LobbyNotificationPauseKeyName, MoveTemp(ClearedPauses));
-
-			// Mark the lobby active for each member.
-			// Lobby activation / deactivation will dispatch lobby join / leave notifications.
-			for (const FJoinLobbyLocalUserData& LocalUser : Params.LocalUsers)
-			{
-				AddActiveLobby(LocalUser.LocalUserId, LobbyData.ToSharedRef());
-			}
+			AddActiveLobby(LocalUser.LocalUserId, LobbyData.ToSharedRef());
+		}
+
+		// todo: figure out a butter way to handle attribute parameters.
+		FClientLobbyDataChanges LobbyChanges;
+		LobbyChanges.LocalName = InAsyncOp.GetParams().LocalName;
 
-			InAsyncOp.SetResult(FJoinLobby::Result{ LobbyData->GetLobbyImpl()});
+		for (const FJoinLobbyLocalUserData& LocalUser : Params.LocalUsers)
+		{
+			TSharedRef<FClientLobbyMemberDataChanges> MemberDataChanges = MakeShared<FClientLobbyMemberDataChanges>();
+			MemberDataChanges->MutatedAttributes = LocalUser.Attributes;
+			LobbyChanges.MutatedMembers.Add(LocalUser.LocalUserId, MoveTemp(MemberDataChanges));
 		}
+		LobbyData->GetClientLobbyData()->ApplyLobbyUpdateFromLocalChanges(MoveTemp(LobbyChanges), LobbyEvents);
+
+		InAsyncOp.SetResult(FJoinLobby::Result{ LobbyData->GetClientLobbyData()->GetPublicDataPtr()});
 	})
 	.Enqueue(GetSerialQueue());
 
@@ -374,6 +433,8 @@ TOnlineAsyncOpHandle<FLeaveLobby> FLobbiesEOS::LeaveLobby(FLeaveLobby::Params&&
 
 	Op->Then([this, LobbyData](TOnlineAsyncOp<FLeaveLobby>& InAsyncOp)
 	{
+		// Remove the user from the EOS lobby.
+
 		const FLeaveLobby::Params& Params = InAsyncOp.GetParams();
 
 		FLeaveLobbyImpl::Params LeaveParams;
@@ -383,11 +444,18 @@ TOnlineAsyncOpHandle<FLeaveLobby> FLobbiesEOS::LeaveLobby(FLeaveLobby::Params&&
 	})
 	.Then([this, LobbyData](TOnlineAsyncOp<FLeaveLobby>& InAsyncOp, TDefaultErrorResult<FLeaveLobbyImpl> Result)
 	{
-		const auto& Params = InAsyncOp.GetParams();
+		const FLeaveLobby::Params& Params = InAsyncOp.GetParams();
+
+		// Remove the user from the local lobby data and dispatch notifications.
+		FClientLobbyDataChanges LobbyChanges;
+		LobbyChanges.LeavingMembers.Add(Params.LocalUserId, ELobbyMemberLeaveReason::Left);
+		FApplyLobbyUpdateResult ApplyResult = LobbyData->GetClientLobbyData()->ApplyLobbyUpdateFromLocalChanges(MoveTemp(LobbyChanges), LobbyEvents);
 
-		if (LobbyData)
+		// Remove the lobby from the active list for the user.
+		// The lobby data will be cleaned up once all references are removed.
+		for (FOnlineAccountIdHandle LeavingMember : ApplyResult.LeavingLocalMembers)
 		{
-			RemoveActiveLobby(Params.LocalUserId, LobbyData.ToSharedRef());
+			RemoveActiveLobby(LeavingMember, LobbyData.ToSharedRef());
 		}
 
 		if (Result.IsError())
@@ -438,7 +506,7 @@ TOnlineAsyncOpHandle<FInviteLobbyMember> FLobbiesEOS::InviteLobbyMember(FInviteL
 			InAsyncOp.SetResult(FInviteLobbyMember::Result{});
 		}
 	})
-	.Enqueue(GetSerialQueue(Params.LocalUserId));
+	.Enqueue(GetSerialQueue());
 
 	return Op->GetHandle();
 }
@@ -475,7 +543,7 @@ TOnlineAsyncOpHandle<FDeclineLobbyInvitation> FLobbiesEOS::DeclineLobbyInvitatio
 			InAsyncOp.SetResult(FDeclineLobbyInvitation::Result{});
 		}
 	})
-	.Enqueue(GetSerialQueue(Params.LocalUserId));
+	.Enqueue(GetSerialQueue());
 
 	return Op->GetHandle();
 }
@@ -491,12 +559,21 @@ TOnlineAsyncOpHandle<FKickLobbyMember> FLobbiesEOS::KickLobbyMember(FKickLobbyMe
 		return Op->GetHandle();
 	}
 
+	TSharedPtr<FLobbyDataEOS> LobbyData = LobbyDataRegistry->Find(Params.LobbyId);
+	if (!LobbyData)
+	{
+		Op->SetError(Errors::InvalidParams());
+		return Op->GetHandle();
+	}
+	Op->Data.Set<TSharedRef<FLobbyDataEOS>>(LobbyDataKeyName, LobbyData.ToSharedRef());
+
 	Op->Then([this](TOnlineAsyncOp<FKickLobbyMember>& InAsyncOp)
 	{
 		const FKickLobbyMember::Params& Params = InAsyncOp.GetParams();
+		const TSharedRef<FLobbyDataEOS>& LobbyData = GetOpDataChecked<TSharedRef<FLobbyDataEOS>>(InAsyncOp, LobbyDataKeyName);
 
 		FKickLobbyMemberImpl::Params KickParams;
-		KickParams.LobbyData = LobbyDataRegistry->Find(Params.LobbyId);
+		KickParams.LobbyData = LobbyData;
 		KickParams.LocalUserId = Params.LocalUserId;
 		KickParams.TargetUserId = Params.TargetUserId;
 		return KickLobbyMemberImpl(MoveTemp(KickParams));
@@ -508,12 +585,19 @@ TOnlineAsyncOpHandle<FKickLobbyMember> FLobbiesEOS::KickLobbyMember(FKickLobbyMe
 			// TODO: Error codes
 			InAsyncOp.SetError(Errors::Unknown(Result.GetErrorValue()));
 		}
-		else
-		{
-			InAsyncOp.SetResult(FKickLobbyMember::Result{});
-		}
 	})
-	.Enqueue(GetSerialQueue(Params.LocalUserId));
+	.Then([this](TOnlineAsyncOp<FKickLobbyMember>& InAsyncOp)
+	{
+		const FKickLobbyMember::Params& Params = InAsyncOp.GetParams();
+		const TSharedRef<FLobbyDataEOS>& LobbyData = GetOpDataChecked<TSharedRef<FLobbyDataEOS>>(InAsyncOp, LobbyDataKeyName);
+
+		// Update local cache and fire events.
+		FClientLobbyDataChanges LobbyChanges;
+		LobbyChanges.LeavingMembers.Add(Params.TargetUserId, ELobbyMemberLeaveReason::Kicked);
+		LobbyData->GetClientLobbyData()->ApplyLobbyUpdateFromLocalChanges(MoveTemp(LobbyChanges), LobbyEvents);
+		InAsyncOp.SetResult(FKickLobbyMember::Result{});
+	})
+	.Enqueue(GetSerialQueue());
 
 	return Op->GetHandle();
 }
@@ -529,12 +613,21 @@ TOnlineAsyncOpHandle<FPromoteLobbyMember> FLobbiesEOS::PromoteLobbyMember(FPromo
 		return Op->GetHandle();
 	}
 
+	TSharedPtr<FLobbyDataEOS> LobbyData = LobbyDataRegistry->Find(Params.LobbyId);
+	if (!LobbyData)
+	{
+		Op->SetError(Errors::InvalidParams());
+		return Op->GetHandle();
+	}
+	Op->Data.Set<TSharedRef<FLobbyDataEOS>>(LobbyDataKeyName, LobbyData.ToSharedRef());
+
 	Op->Then([this](TOnlineAsyncOp<FPromoteLobbyMember>& InAsyncOp)
 	{
 		const FPromoteLobbyMember::Params& Params = InAsyncOp.GetParams();
+		const TSharedRef<FLobbyDataEOS>& LobbyData = GetOpDataChecked<TSharedRef<FLobbyDataEOS>>(InAsyncOp, LobbyDataKeyName);
 
 		FPromoteLobbyMemberImpl::Params PromoteParams;
-		PromoteParams.LobbyData = LobbyDataRegistry->Find(Params.LobbyId);
+		PromoteParams.LobbyData = LobbyData;
 		PromoteParams.LocalUserId = Params.LocalUserId;
 		PromoteParams.TargetUserId = Params.TargetUserId;
 		return PromoteLobbyMemberImpl(MoveTemp(PromoteParams));
@@ -546,12 +639,19 @@ TOnlineAsyncOpHandle<FPromoteLobbyMember> FLobbiesEOS::PromoteLobbyMember(FPromo
 			// TODO: Error codes
 			InAsyncOp.SetError(Errors::Unknown(Result.GetErrorValue()));
 		}
-		else
-		{
-			InAsyncOp.SetResult(FPromoteLobbyMember::Result{});
-		}
 	})
-	.Enqueue(GetSerialQueue(Params.LocalUserId));
+	.Then([this](TOnlineAsyncOp<FPromoteLobbyMember>& InAsyncOp)
+	{
+		const FPromoteLobbyMember::Params& Params = InAsyncOp.GetParams();
+		const TSharedRef<FLobbyDataEOS>& LobbyData = GetOpDataChecked<TSharedRef<FLobbyDataEOS>>(InAsyncOp, LobbyDataKeyName);
+
+		// Update local cache and fire events.
+		FClientLobbyDataChanges LobbyChanges;
+		LobbyChanges.OwnerAccountId = Params.TargetUserId;
+		LobbyData->GetClientLobbyData()->ApplyLobbyUpdateFromLocalChanges(MoveTemp(LobbyChanges), LobbyEvents);
+		InAsyncOp.SetResult(FPromoteLobbyMember::Result{});
+	})
+	.Enqueue(GetSerialQueue());
 
 	return Op->GetHandle();
 }
@@ -567,15 +667,27 @@ TOnlineAsyncOpHandle<FModifyLobbyJoinPolicy> FLobbiesEOS::ModifyLobbyJoinPolicy(
 		return Op->GetHandle();
 	}
 
+	TSharedPtr<FLobbyDataEOS> LobbyData = LobbyDataRegistry->Find(Params.LobbyId);
+	if (!LobbyData)
+	{
+		Op->SetError(Errors::InvalidParams());
+		return Op->GetHandle();
+	}
+	Op->Data.Set<TSharedRef<FLobbyDataEOS>>(LobbyDataKeyName, LobbyData.ToSharedRef());
+
 	Op->Then([this](TOnlineAsyncOp<FModifyLobbyJoinPolicy>& InAsyncOp)
 	{
 		const FModifyLobbyJoinPolicy::Params& Params = InAsyncOp.GetParams();
+		const TSharedRef<FLobbyDataEOS>& LobbyData = GetOpDataChecked<TSharedRef<FLobbyDataEOS>>(InAsyncOp, LobbyDataKeyName);
+
+		TSharedRef<FClientLobbyDataChanges> LobbyChanges = MakeShared<FClientLobbyDataChanges>();
+		InAsyncOp.Data.Set<TSharedRef<FClientLobbyDataChanges>>(LobbyChangesKeyName, LobbyChanges);
+		LobbyChanges->JoinPolicy = Params.JoinPolicy;
 
 		FModifyLobbyDataImpl::Params ModifyLobbyDataParams;
-		ModifyLobbyDataParams.LobbyData = LobbyDataRegistry->Find(Params.LobbyId);
+		ModifyLobbyDataParams.LobbyData = LobbyData;
 		ModifyLobbyDataParams.LocalUserId = Params.LocalUserId;
-		ModifyLobbyDataParams.Changes = MakeShared<FLobbyDataChanges>();
-		ModifyLobbyDataParams.Changes->JoinPolicy = Params.JoinPolicy;
+		ModifyLobbyDataParams.Changes = LobbyChanges;
 		return ModifyLobbyDataImpl(MoveTemp(ModifyLobbyDataParams));
 	})
 	.Then([this](TOnlineAsyncOp<FModifyLobbyJoinPolicy>& InAsyncOp, TDefaultErrorResult<FModifyLobbyDataImpl> Result)
@@ -585,12 +697,17 @@ TOnlineAsyncOpHandle<FModifyLobbyJoinPolicy> FLobbiesEOS::ModifyLobbyJoinPolicy(
 			// TODO: Error codes
 			InAsyncOp.SetError(Errors::Unknown(Result.GetErrorValue()));
 		}
-		else
-		{
-			InAsyncOp.SetResult(FModifyLobbyJoinPolicy::Result{});
-		}
 	})
-	.Enqueue(GetSerialQueue(Params.LocalUserId));
+	.Then([this](TOnlineAsyncOp<FModifyLobbyJoinPolicy>& InAsyncOp)
+	{
+		const TSharedRef<FLobbyDataEOS>& LobbyData = GetOpDataChecked<TSharedRef<FLobbyDataEOS>>(InAsyncOp, LobbyDataKeyName);
+		const TSharedRef<FClientLobbyDataChanges>& LobbyChanges = GetOpDataChecked<TSharedRef<FClientLobbyDataChanges>>(InAsyncOp, LobbyChangesKeyName);
+
+		// Update local cache and fire events.
+		LobbyData->GetClientLobbyData()->ApplyLobbyUpdateFromLocalChanges(MoveTemp(*LobbyChanges), LobbyEvents);
+		InAsyncOp.SetResult(FModifyLobbyJoinPolicy::Result{});
+	})
+	.Enqueue(GetSerialQueue());
 
 	return Op->GetHandle();
 }
@@ -606,16 +723,28 @@ TOnlineAsyncOpHandle<FModifyLobbyAttributes> FLobbiesEOS::ModifyLobbyAttributes(
 		return Op->GetHandle();
 	}
 
+	TSharedPtr<FLobbyDataEOS> LobbyData = LobbyDataRegistry->Find(Params.LobbyId);
+	if (!LobbyData)
+	{
+		Op->SetError(Errors::InvalidParams());
+		return Op->GetHandle();
+	}
+	Op->Data.Set<TSharedRef<FLobbyDataEOS>>(LobbyDataKeyName, LobbyData.ToSharedRef());
+
 	Op->Then([this](TOnlineAsyncOp<FModifyLobbyAttributes>& InAsyncOp)
 	{
 		const FModifyLobbyAttributes::Params& Params = InAsyncOp.GetParams();
+		const TSharedRef<FLobbyDataEOS>& LobbyData = GetOpDataChecked<TSharedRef<FLobbyDataEOS>>(InAsyncOp, LobbyDataKeyName);
+
+		TSharedRef<FClientLobbyDataChanges> LobbyChanges = MakeShared<FClientLobbyDataChanges>();
+		InAsyncOp.Data.Set<TSharedRef<FClientLobbyDataChanges>>(LobbyChangesKeyName, LobbyChanges);
+		LobbyChanges->MutatedAttributes = Params.MutatedAttributes;
+		LobbyChanges->ClearedAttributes = Params.ClearedAttributes;
 
 		FModifyLobbyDataImpl::Params ModifyLobbyDataParams;
 		ModifyLobbyDataParams.LobbyData = LobbyDataRegistry->Find(Params.LobbyId);
 		ModifyLobbyDataParams.LocalUserId = Params.LocalUserId;
-		ModifyLobbyDataParams.Changes = MakeShared<FLobbyDataChanges>();
-		ModifyLobbyDataParams.Changes->MutatedAttributes = Params.MutatedAttributes;
-		ModifyLobbyDataParams.Changes->ClearedAttributes = Params.ClearedAttributes;
+		ModifyLobbyDataParams.Changes = LobbyChanges;
 		return ModifyLobbyDataImpl(MoveTemp(ModifyLobbyDataParams));
 	})
 	.Then([this](TOnlineAsyncOp<FModifyLobbyAttributes>& InAsyncOp, TDefaultErrorResult<FModifyLobbyDataImpl> Result)
@@ -625,12 +754,17 @@ TOnlineAsyncOpHandle<FModifyLobbyAttributes> FLobbiesEOS::ModifyLobbyAttributes(
 			// TODO: Error codes
 			InAsyncOp.SetError(Errors::Unknown(Result.GetErrorValue()));
 		}
-		else
-		{
-			InAsyncOp.SetResult(FModifyLobbyAttributes::Result{});
-		}
 	})
-	.Enqueue(GetSerialQueue(Params.LocalUserId));
+	.Then([this](TOnlineAsyncOp<FModifyLobbyAttributes>& InAsyncOp)
+	{
+		const TSharedRef<FLobbyDataEOS>& LobbyData = GetOpDataChecked<TSharedRef<FLobbyDataEOS>>(InAsyncOp, LobbyDataKeyName);
+		const TSharedRef<FClientLobbyDataChanges>& LobbyChanges = GetOpDataChecked<TSharedRef<FClientLobbyDataChanges>>(InAsyncOp, LobbyChangesKeyName);
+
+		// Update local cache and fire events.
+		LobbyData->GetClientLobbyData()->ApplyLobbyUpdateFromLocalChanges(MoveTemp(*LobbyChanges), LobbyEvents);
+		InAsyncOp.SetResult(FModifyLobbyAttributes::Result{});
+	})
+	.Enqueue(GetSerialQueue());
 
 	return Op->GetHandle();
 }
@@ -646,16 +780,28 @@ TOnlineAsyncOpHandle<FModifyLobbyMemberAttributes> FLobbiesEOS::ModifyLobbyMembe
 		return Op->GetHandle();
 	}
 
+	TSharedPtr<FLobbyDataEOS> LobbyData = LobbyDataRegistry->Find(Params.LobbyId);
+	if (!LobbyData)
+	{
+		Op->SetError(Errors::InvalidParams());
+		return Op->GetHandle();
+	}
+	Op->Data.Set<TSharedRef<FLobbyDataEOS>>(LobbyDataKeyName, LobbyData.ToSharedRef());
+
 	Op->Then([this](TOnlineAsyncOp<FModifyLobbyMemberAttributes>& InAsyncOp)
 	{
 		const FModifyLobbyMemberAttributes::Params& Params = InAsyncOp.GetParams();
+		const TSharedRef<FLobbyDataEOS>& LobbyData = GetOpDataChecked<TSharedRef<FLobbyDataEOS>>(InAsyncOp, LobbyDataKeyName);
+
+		TSharedRef<FClientLobbyMemberDataChanges> LobbyMemberChanges = MakeShared<FClientLobbyMemberDataChanges>();
+		InAsyncOp.Data.Set<TSharedRef<FClientLobbyMemberDataChanges>>(LobbyMemberChangesKeyName, LobbyMemberChanges);
+		LobbyMemberChanges->MutatedAttributes = Params.MutatedAttributes;
+		LobbyMemberChanges->ClearedAttributes = Params.ClearedAttributes;
 
 		FModifyLobbyMemberDataImpl::Params ModifyLobbyMemberDataParams;
-		ModifyLobbyMemberDataParams.LobbyData = LobbyDataRegistry->Find(Params.LobbyId);
+		ModifyLobbyMemberDataParams.LobbyData = LobbyData;
 		ModifyLobbyMemberDataParams.LocalUserId = Params.LocalUserId;
-		ModifyLobbyMemberDataParams.Changes = MakeShared<FLobbyMemberDataChanges>();
-		ModifyLobbyMemberDataParams.Changes->MutatedAttributes = Params.MutatedAttributes;
-		ModifyLobbyMemberDataParams.Changes->ClearedAttributes = Params.ClearedAttributes;
+		ModifyLobbyMemberDataParams.Changes = LobbyMemberChanges;
 		return ModifyLobbyMemberDataImpl(MoveTemp(ModifyLobbyMemberDataParams));
 	})
 	.Then([this](TOnlineAsyncOp<FModifyLobbyMemberAttributes>& InAsyncOp, TDefaultErrorResult<FModifyLobbyMemberDataImpl> Result)
@@ -665,12 +811,20 @@ TOnlineAsyncOpHandle<FModifyLobbyMemberAttributes> FLobbiesEOS::ModifyLobbyMembe
 			// TODO: Error codes
 			InAsyncOp.SetError(Errors::Unknown(Result.GetErrorValue()));
 		}
-		else
-		{
-			InAsyncOp.SetResult(FModifyLobbyMemberAttributes::Result{});
-		}
 	})
-	.Enqueue(GetSerialQueue(Params.LocalUserId));
+	.Then([this](TOnlineAsyncOp<FModifyLobbyMemberAttributes>& InAsyncOp)
+	{
+		const FModifyLobbyMemberAttributes::Params& Params = InAsyncOp.GetParams();
+		const TSharedRef<FLobbyDataEOS>& LobbyData = GetOpDataChecked<TSharedRef<FLobbyDataEOS>>(InAsyncOp, LobbyDataKeyName);
+		const TSharedRef<FClientLobbyMemberDataChanges>& LobbyMemberChanges = GetOpDataChecked<TSharedRef<FClientLobbyMemberDataChanges>>(InAsyncOp, LobbyMemberChangesKeyName);
+
+		// Update local cache and fire events.
+		FClientLobbyDataChanges LobbyChanges;
+		LobbyChanges.MutatedMembers.Add(Params.LocalUserId, LobbyMemberChanges);
+		LobbyData->GetClientLobbyData()->ApplyLobbyUpdateFromLocalChanges(MoveTemp(LobbyChanges), LobbyEvents);
+		InAsyncOp.SetResult(FModifyLobbyMemberAttributes::Result{});
+	})
+	.Enqueue(GetSerialQueue());
 
 	return Op->GetHandle();
 }
@@ -683,7 +837,7 @@ TOnlineResult<FGetJoinedLobbies> FLobbiesEOS::GetJoinedLobbies(FGetJoinedLobbies
 		Result.Lobbies.Reserve(Lobbies->Num());
 		for (const TSharedRef<FLobbyDataEOS>& LobbyDataEOS : *Lobbies)
 		{
-			Result.Lobbies.Emplace(LobbyDataEOS->GetLobbyImpl());
+			Result.Lobbies.Emplace(LobbyDataEOS->GetClientLobbyData()->GetPublicDataPtr());
 		}
 		return TOnlineResult<FGetJoinedLobbies>(MoveTemp(Result));
 	}
@@ -695,68 +849,88 @@ TOnlineResult<FGetJoinedLobbies> FLobbiesEOS::GetJoinedLobbies(FGetJoinedLobbies
 
 void FLobbiesEOS::HandleLobbyUpdated(const EOS_Lobby_LobbyUpdateReceivedCallbackInfo* Data)
 {
-#if 0
-	EOS_STRUCT(EOS_Lobby_LobbyUpdateReceivedCallbackInfo, (
-		/** Context that was passed into EOS_Lobby_AddNotifyLobbyUpdateReceived */
-		void* ClientData;
-		/** The ID of the lobby */
-		EOS_LobbyId LobbyId;
-	));
-#endif
-
 	if (TSharedPtr<FLobbyDataEOS> LobbyData = LobbyDataRegistry->Find(Data->LobbyId))
 	{
-		UE_LOG(LogTemp, Warning, TEXT("[FLobbiesEOS::HandleLobbyUpdated] Received update. Lobby: %s"), *LobbyData->GetLobbyId());
+		FProcessLobbyNotificationImpl::Params Params;
+		Params.LobbyData = LobbyData;
+
+		ProcessLobbyNotificationImplOp(MoveTemp(Params))
+		.OnComplete([LobbyData](const TOnlineResult<FProcessLobbyNotificationImpl>& Result)
+		{
+			if (Result.IsError())
+			{
+				// Todo: handle failure to update lobby from snapshot.
+				UE_LOG(LogTemp, Warning, TEXT("[FLobbiesEOS::HandleLobbyUpdated] Failed to apply update. Lobby: %s, Error: %s"), *LobbyData->GetLobbyId(), *Result.GetErrorValue().GetLogString());
+			}
+		});
 	}
 }
 
 void FLobbiesEOS::HandleLobbyMemberUpdated(const EOS_Lobby_LobbyMemberUpdateReceivedCallbackInfo* Data)
 {
-#if 0
-	EOS_STRUCT(EOS_Lobby_LobbyMemberUpdateReceivedCallbackInfo, (
-		/** Context that was passed into EOS_Lobby_AddNotifyLobbyMemberUpdateReceived */
-		void* ClientData;
-		/** The ID of the lobby */
-		EOS_LobbyId LobbyId;
-		/** The Product User ID of the lobby member */
-		EOS_ProductUserId TargetUserId;
-	));
-#endif
-
-	const FOnlineAccountIdHandle LocalUserId = FindAccountId(Data->TargetUserId);
-	if (LocalUserId.IsValid())
+	if (TSharedPtr<FLobbyDataEOS> LobbyData = LobbyDataRegistry->Find(Data->LobbyId))
 	{
-		if (TSharedPtr<FLobbyDataEOS> LobbyData = LobbyDataRegistry->Find(Data->LobbyId))
+		FProcessLobbyNotificationImpl::Params Params;
+		Params.LobbyData = LobbyData;
+		Params.MutatedMembers.Add(Data->TargetUserId);
+
+		ProcessLobbyNotificationImplOp(MoveTemp(Params))
+		.OnComplete([LobbyData](const TOnlineResult<FProcessLobbyNotificationImpl>& Result)
 		{
-			UE_LOG(LogTemp, Warning, TEXT("[FLobbiesEOS::HandleLobbyMemberUpdated] Received update. Lobby: %s, Member: %s"),
-				*LobbyData->GetLobbyId(), *ToLogString(LocalUserId));
-		}
+			if (Result.IsError())
+			{
+				// Todo: handle failure to update lobby from snapshot.
+				UE_LOG(LogTemp, Warning, TEXT("[FLobbiesEOS::HandleLobbyMemberUpdated] Failed to apply update. Lobby: %s, Error: %s"), *LobbyData->GetLobbyId(), *Result.GetErrorValue().GetLogString());
+			}
+		});
 	}
 }
 
 void FLobbiesEOS::HandleLobbyMemberStatusReceived(const EOS_Lobby_LobbyMemberStatusReceivedCallbackInfo* Data)
 {
-#if 0
-	EOS_STRUCT(EOS_Lobby_LobbyMemberStatusReceivedCallbackInfo, (
-		/** Context that was passed into EOS_Lobby_AddNotifyLobbyMemberStatusReceived */
-		void* ClientData;
-		/** The ID of the lobby */
-		EOS_LobbyId LobbyId;
-		/** The Product User ID of the lobby member */
-		EOS_ProductUserId TargetUserId;
-		/** Latest status of the user */
-		EOS_ELobbyMemberStatus CurrentStatus;
-	));
-#endif
-
-	const FOnlineAccountIdHandle LocalUserId = FindAccountId(Data->TargetUserId);
-	if (LocalUserId.IsValid())
+	if (TSharedPtr<FLobbyDataEOS> LobbyData = LobbyDataRegistry->Find(Data->LobbyId))
 	{
-		if (TSharedPtr<FLobbyDataEOS> LobbyData = LobbyDataRegistry->Find(Data->LobbyId))
+		FProcessLobbyNotificationImpl::Params Params;
+		Params.LobbyData = LobbyData;
+
+		switch (Data->CurrentStatus)
 		{
-			UE_LOG(LogTemp, Warning, TEXT("[FLobbiesEOS::HandleLobbyMemberStatusReceived] Received update. Lobby: %s, Member: %s, Statsu: %d"),
-				*LobbyData->GetLobbyId(), *ToLogString(LocalUserId), Data->CurrentStatus);
+		case EOS_ELobbyMemberStatus::EOS_LMS_JOINED:
+			// Fetch member snapshot on join.
+			Params.MutatedMembers.Add(Data->TargetUserId);
+			break;
+
+		case EOS_ELobbyMemberStatus::EOS_LMS_PROMOTED:
+			// No member data needed, only lobby snapshot.
+			break;
+
+		case EOS_ELobbyMemberStatus::EOS_LMS_LEFT:
+			Params.LeavingMembers.Add(Data->TargetUserId, ELobbyMemberLeaveReason::Left);
+			break;
+
+		case EOS_ELobbyMemberStatus::EOS_LMS_DISCONNECTED:
+			Params.LeavingMembers.Add(Data->TargetUserId, ELobbyMemberLeaveReason::Disconnected);
+			break;
+
+		case EOS_ELobbyMemberStatus::EOS_LMS_KICKED:
+			Params.LeavingMembers.Add(Data->TargetUserId, ELobbyMemberLeaveReason::Kicked);
+			break;
+
+		default: checkNoEntry(); // Intentional fallthrough
+		case EOS_ELobbyMemberStatus::EOS_LMS_CLOSED:
+			Params.LeavingMembers.Add(Data->TargetUserId, ELobbyMemberLeaveReason::Closed);
+			break;
 		}
+
+		ProcessLobbyNotificationImplOp(MoveTemp(Params))
+		.OnComplete([LobbyData](const TOnlineResult<FProcessLobbyNotificationImpl>& Result)
+		{
+			if (Result.IsError())
+			{
+				// Todo: handle failure to update lobby from snapshot.
+				UE_LOG(LogTemp, Warning, TEXT("[FLobbiesEOS::HandleLobbyMemberUpdated] Failed to apply update. Lobby: %s, Error: %s"), *LobbyData->GetLobbyId(), *Result.GetErrorValue().GetLogString());
+			}
+		});
 	}
 }
 
@@ -780,7 +954,7 @@ void FLobbiesEOS::HandleLobbyInviteReceived(const EOS_Lobby_LobbyInviteReceivedC
 	if (LocalUserId.IsValid())
 	{
 		FLobbyInviteDataEOS::CreateFromInviteId(LobbyPrerequisites.ToSharedRef(), LobbyDataRegistry.ToSharedRef(), LocalUserId, Data->InviteId, Data->TargetUserId)
-		.Then([this](TFuture<TDefaultErrorResultInternal<TSharedPtr<FLobbyInviteDataEOS>>>&& Future)
+		.Then([this](TFuture<TDefaultErrorResultInternal<TSharedRef<FLobbyInviteDataEOS>>>&& Future)
 		{
 			if (Future.Get().IsError())
 			{
@@ -795,7 +969,7 @@ void FLobbiesEOS::HandleLobbyInviteReceived(const EOS_Lobby_LobbyInviteReceivedC
 					*Future.Get().GetOkValue()->GetLobbyData()->GetLobbyId(),
 					*ToLogString(Future.Get().GetOkValue()->GetReceiver()),
 					*ToLogString(Future.Get().GetOkValue()->GetSender()));
-				AddActiveInvite(Future.Get().GetOkValue().ToSharedRef());
+				AddActiveInvite(Future.Get().GetOkValue());
 			}
 		});
 	}
@@ -910,45 +1084,16 @@ void FLobbiesEOS::UnregisterHandlers()
 
 void FLobbiesEOS::AddActiveLobby(FOnlineAccountIdHandle LocalUserId, const TSharedRef<FLobbyDataEOS>& LobbyData)
 {
-	if (LobbyData->GetLocalMembers().IsEmpty())
-	{
-		// Notify lobby added.
-		OnLobbyJoinedEvent.Broadcast(FLobbyJoined{LobbyData->GetLobbyImpl()});
-	}
-
+	// Add bookkeeping for the user.
 	ActiveLobbies.FindOrAdd(LocalUserId).Add(LobbyData);
-	LobbyData->GetLocalMembers().Add(LocalUserId);
-
-	// Todo: make this better.
-	for (const TSharedRef<const FLobbyMember>& MemberData : LobbyData->GetLobbyImpl()->Members)
-	{
-		if (MemberData->AccountId == LocalUserId)
-		{
-			OnLobbyMemberJoinedEvent.Broadcast(FLobbyMemberJoined{LobbyData->GetLobbyImpl(), MemberData});
-			break;
-		}
-	}
 }
 
 void FLobbiesEOS::RemoveActiveLobby(FOnlineAccountIdHandle LocalUserId, const TSharedRef<FLobbyDataEOS>& LobbyData)
 {
-	// Todo: make this better.
-	for (const TSharedRef<const FLobbyMember>& MemberData : LobbyData->GetLobbyImpl()->Members)
-	{
-		if (MemberData->AccountId == LocalUserId)
-		{
-			OnLobbyMemberLeftEvent.Broadcast(FLobbyMemberLeft{LobbyData->GetLobbyImpl(), MemberData});
-			break;
-		}
-	}
-
-	ActiveLobbies.FindOrAdd(LocalUserId).Remove(LobbyData);
-	LobbyData->GetLocalMembers().Remove(LocalUserId);
-
-	if (LobbyData->GetLocalMembers().IsEmpty())
+	// Remove bookkeeping for the local user.
+	if (TSet<TSharedRef<FLobbyDataEOS>>* Lobbies = ActiveLobbies.Find(LocalUserId))
 	{
-		// Notify lobby removed.
-		OnLobbyLeftEvent.Broadcast(FLobbyLeft{LobbyData->GetLobbyImpl()});
+		Lobbies->Remove(LobbyData);
 	}
 }
 
@@ -961,11 +1106,11 @@ void FLobbiesEOS::AddActiveInvite(const TSharedRef<FLobbyInviteDataEOS>& Invite)
 	if (ActiveUserInvites.Find(LobbyId) == nullptr)
 	{
 		ActiveUserInvites.Add(LobbyId, Invite);
-		OnLobbyInvitationAddedEvent.Broadcast(
+		LobbyEvents.OnLobbyInvitationAdded.Broadcast(
 			FLobbyInvitationAdded{
 				Invite->GetReceiver(),
 				Invite->GetSender(),
-				Invite->GetLobbyData()->GetLobbyImpl()
+				Invite->GetLobbyData()->GetClientLobbyData()->GetPublicDataPtr()
 			});
 	}
 }
@@ -974,11 +1119,11 @@ void FLobbiesEOS::RemoveActiveInvite(const TSharedRef<FLobbyInviteDataEOS>& Invi
 {
 	ActiveInvites.FindOrAdd(Invite->GetReceiver()).Remove(Invite->GetLobbyData()->GetLobbyIdHandle());
 
-	OnLobbyInvitationRemovedEvent.Broadcast(
+	LobbyEvents.OnLobbyInvitationRemoved.Broadcast(
 		FLobbyInvitationRemoved{
 			Invite->GetReceiver(),
 			Invite->GetSender(),
-			Invite->GetLobbyData()->GetLobbyImpl()
+			Invite->GetLobbyData()->GetClientLobbyData()->GetPublicDataPtr()
 		});
 }
 
@@ -1021,7 +1166,7 @@ TFuture<TDefaultErrorResult<FLobbiesEOS::FJoinLobbyImpl>> FLobbiesEOS::JoinLobby
 		JoinLobbyMemberParams.LocalUserId = UserData.LocalUserId;
 		JoinLobbyMemberParams.Attributes = MoveTemp(UserData.Attributes);
 
-		auto JoinMemberPromise = MakeShared<TPromise<TDefaultErrorResult<FJoinLobbyMemberImpl>>>();
+		TSharedRef<TPromise<TDefaultErrorResult<FJoinLobbyMemberImpl>>> JoinMemberPromise = MakeShared<TPromise<TDefaultErrorResult<FJoinLobbyMemberImpl>>>();
 		JoinLobbyMemberImplOp(MoveTemp(JoinLobbyMemberParams))
 		.OnComplete([JoinMemberPromise](const TOnlineResult<FJoinLobbyMemberImpl>& Result)
 		{
@@ -1039,7 +1184,7 @@ TFuture<TDefaultErrorResult<FLobbiesEOS::FJoinLobbyImpl>> FLobbiesEOS::JoinLobby
 	}
 
 	TPromise<TDefaultErrorResult<FJoinLobbyImpl>> Promise;
-	auto Future = Promise.GetFuture();
+	TFuture<TDefaultErrorResult<FJoinLobbyImpl>> Future = Promise.GetFuture();
 
 	WhenAll(MoveTemp(PendingMemberJoins))
 	.Then([this, Promise = MoveTemp(Promise), Params = MoveTemp(Params)](TFuture<TArray<TDefaultErrorResult<FJoinLobbyMemberImpl>>>&& Results) mutable
@@ -1114,7 +1259,7 @@ TOnlineAsyncOpHandle<FLobbiesEOS::FJoinLobbyMemberImpl> FLobbiesEOS::JoinLobbyMe
 
 	Op->Then([this, LobbyDetails, LobbyData = Params.LobbyData](TOnlineAsyncOp<FJoinLobbyMemberImpl>& InAsyncOp)
 	{
-		const auto& Params = InAsyncOp.GetParams();
+		const FLobbiesEOS::FJoinLobbyMemberImpl::Params& Params = InAsyncOp.GetParams();
 
 		EOS_Lobby_JoinLobbyOptions JoinLobbyOptions = {};
 		JoinLobbyOptions.ApiVersion = EOS_LOBBY_JOINLOBBY_API_LATEST;
@@ -1122,15 +1267,10 @@ TOnlineAsyncOpHandle<FLobbiesEOS::FJoinLobbyMemberImpl> FLobbiesEOS::JoinLobbyMe
 		JoinLobbyOptions.LocalUserId = GetProductUserIdChecked(Params.LocalUserId);
 		JoinLobbyOptions.bPresenceEnabled = false;
 		JoinLobbyOptions.LocalRTCOptions = nullptr;
-
-		UE_LOG(LogTemp, Warning, TEXT("[FLobbiesEOS::EOS_Lobby_JoinLobby] Start. Lobby: %s, Member: %s"),
-			*LobbyData->GetLobbyId(), *ToLogString(Params.LocalUserId));
-
 		return EOS_Async<EOS_Lobby_JoinLobbyCallbackInfo>(EOS_Lobby_JoinLobby, LobbyPrerequisites->LobbyInterfaceHandle, JoinLobbyOptions);
 	})
 	.Then([this](TOnlineAsyncOp<FJoinLobbyMemberImpl>& InAsyncOp, const EOS_Lobby_JoinLobbyCallbackInfo* Data)
 	{
-		UE_LOG(LogTemp, Warning, TEXT("[FLobbiesEOS::EOS_Lobby_JoinLobby] Complete. Result: %d"), Data->ResultCode);
 		if (Data->ResultCode != EOS_EResult::EOS_Success)
 		{
 			// TODO: Error codes
@@ -1139,17 +1279,17 @@ TOnlineAsyncOpHandle<FLobbiesEOS::FJoinLobbyMemberImpl> FLobbiesEOS::JoinLobbyMe
 	})
 	.Then([this](TOnlineAsyncOp<FJoinLobbyMemberImpl>& InAsyncOp)
 	{
-		const auto& Params = InAsyncOp.GetParams();
+		const FLobbiesEOS::FJoinLobbyMemberImpl::Params& Params = InAsyncOp.GetParams();
 		FModifyLobbyMemberDataImpl::Params ModifyParams;
 		ModifyParams.LobbyData = Params.LobbyData;
 		ModifyParams.LocalUserId = Params.LocalUserId;
-		ModifyParams.Changes = MakeShared<FLobbyMemberDataChanges>();
+		ModifyParams.Changes = MakeShared<FClientLobbyMemberDataChanges>();
 		ModifyParams.Changes->MutatedAttributes = Params.Attributes;
 		return ModifyLobbyMemberDataImpl(MoveTemp(ModifyParams));
 	})
 	.Then([this](TOnlineAsyncOp<FJoinLobbyMemberImpl>& InAsyncOp, TDefaultErrorResult<FModifyLobbyMemberDataImpl>&& Result)
 	{
-		const auto& Params = InAsyncOp.GetParams();
+		const FLobbiesEOS::FJoinLobbyMemberImpl::Params& Params = InAsyncOp.GetParams();
 
 		if (Result.IsError())
 		{
@@ -1159,7 +1299,7 @@ TOnlineAsyncOpHandle<FLobbiesEOS::FJoinLobbyMemberImpl> FLobbiesEOS::JoinLobbyMe
 			LeaveLobbyParams.LocalUserId = Params.LocalUserId;
 
 			TPromise<void> Promise;
-			auto Future = Promise.GetFuture();
+			TFuture<void> Future = Promise.GetFuture();
 
 			LeaveLobbyImpl(MoveTemp(LeaveLobbyParams))
 			.Then([this, InAsyncOp = InAsyncOp.AsShared(), ErrorResult = MoveTemp(Result.GetErrorValue()), Promise = MoveTemp(Promise)](TFuture<TDefaultErrorResult<FLeaveLobbyImpl>>&& Future) mutable
@@ -1185,7 +1325,7 @@ TOnlineAsyncOpHandle<FLobbiesEOS::FJoinLobbyMemberImpl> FLobbiesEOS::JoinLobbyMe
 	{
 		InAsyncOp.SetResult(FJoinLobbyMemberImpl::Result{});
 	})
-	.Enqueue(GetSerialQueue(Params.LocalUserId));
+	.Enqueue(GetSerialQueue(InParams.LocalUserId));
 
 	return Op->GetHandle();
 }
@@ -1209,7 +1349,7 @@ TFuture<TDefaultErrorResult<FLobbiesEOS::FLeaveLobbyImpl>> FLobbiesEOS::LeaveLob
 	LeaveLobbyOptions.LobbyId = Params.LobbyData->GetLobbyIdEOS();
 
 	TPromise<TDefaultErrorResult<FLeaveLobbyImpl>> Promise;
-	auto Future = Promise.GetFuture();
+	TFuture<TDefaultErrorResult<FLeaveLobbyImpl>> Future = Promise.GetFuture();
 
 	EOS_Async<EOS_Lobby_LeaveLobbyCallbackInfo>(EOS_Lobby_LeaveLobby, LobbyPrerequisites->LobbyInterfaceHandle, LeaveLobbyOptions)
 	.Then([Promise = MoveTemp(Promise)](TFuture<const EOS_Lobby_LeaveLobbyCallbackInfo*>&& Future) mutable
@@ -1247,7 +1387,7 @@ TFuture<TDefaultErrorResult<FLobbiesEOS::FDestroyLobbyImpl>> FLobbiesEOS::Destro
 	DestroyLobbyOptions.LobbyId = Params.LobbyData->GetLobbyIdEOS();
 
 	TPromise<TDefaultErrorResult<FDestroyLobbyImpl>> Promise;
-	auto Future = Promise.GetFuture();
+	TFuture<TDefaultErrorResult<FDestroyLobbyImpl>> Future = Promise.GetFuture();
 
 	EOS_Async<EOS_Lobby_DestroyLobbyCallbackInfo>(EOS_Lobby_DestroyLobby, LobbyPrerequisites->LobbyInterfaceHandle, DestroyLobbyOptions)
 	.Then([Promise = MoveTemp(Promise)](TFuture<const EOS_Lobby_DestroyLobbyCallbackInfo*>&& Future) mutable
@@ -1286,7 +1426,7 @@ TFuture<TDefaultErrorResult<FLobbiesEOS::FInviteLobbyMemberImpl>> FLobbiesEOS::I
 	SendInviteOptions.LobbyId = Params.LobbyData->GetLobbyIdEOS();
 
 	TPromise<TDefaultErrorResult<FInviteLobbyMemberImpl>> Promise;
-	auto Future = Promise.GetFuture();
+	TFuture<TDefaultErrorResult<FInviteLobbyMemberImpl>> Future = Promise.GetFuture();
 
 	EOS_Async<EOS_Lobby_SendInviteCallbackInfo>(EOS_Lobby_SendInvite, LobbyPrerequisites->LobbyInterfaceHandle, SendInviteOptions)
 	.Then([Promise = MoveTemp(Promise)](TFuture<const EOS_Lobby_SendInviteCallbackInfo*>&& Future) mutable
@@ -1444,15 +1584,16 @@ TFuture<TDefaultErrorResult<FLobbiesEOS::FModifyLobbyDataImpl>> FLobbiesEOS::Mod
 		return MakeFulfilledPromise<TDefaultErrorResult<FModifyLobbyDataImpl>>(Errors::InvalidParams()).GetFuture();
 	}
 
-	TDefaultErrorResultInternal<TSharedPtr<FLobbyDetailsEOS>> DetailsResult = FLobbyDetailsEOS::CreateFromLobbyId(
-		LobbyPrerequisites.ToSharedRef(), Params.LobbyData->GetLobbyIdEOS(), Params.LocalUserId);
-	if (DetailsResult.IsError())
+	TSharedPtr<FLobbyDetailsEOS> LobbyDetails = Params.LobbyData->GetUserLobbyDetails(Params.LocalUserId);
+	if (!LobbyDetails)
 	{
-		// Todo: Errors
-		return MakeFulfilledPromise<TDefaultErrorResult<FModifyLobbyDataImpl>>(Errors::Unknown(MoveTemp(DetailsResult.GetErrorValue()))).GetFuture();
+		return MakeFulfilledPromise<TDefaultErrorResult<FModifyLobbyDataImpl>>(Errors::InvalidParams()).GetFuture();
 	}
 
-	// todo: check local user is lobby owner
+	if (Params.LocalUserId != Params.LobbyData->GetClientLobbyData()->GetPublicData().OwnerAccountId)
+	{
+		return MakeFulfilledPromise<TDefaultErrorResult<FModifyLobbyDataImpl>>(Errors::InvalidParams()).GetFuture();
+	}
 
 	UE_LOG(LogTemp, Warning, TEXT("[FLobbiesEOS::ModifyLobbyDataImpl] Start. Lobby: %s, Member: %s"),
 		*Params.LobbyData->GetLobbyId(), *ToLogString(Params.LocalUserId));
@@ -1460,7 +1601,7 @@ TFuture<TDefaultErrorResult<FLobbiesEOS::FModifyLobbyDataImpl>> FLobbiesEOS::Mod
 	TPromise<TDefaultErrorResult<FModifyLobbyDataImpl>> Promise;
 	TFuture<TDefaultErrorResult<FModifyLobbyDataImpl>> Future = Promise.GetFuture();
 
-	DetailsResult.GetOkValue()->ApplyLobbyDataUpdates(Params.LocalUserId, MoveTemp(*Params.Changes))
+	LobbyDetails->ApplyLobbyDataUpdateFromLocalChanges(Params.LocalUserId, *Params.Changes)
 	.Then([Promise = MoveTemp(Promise)](TFuture<EOS_EResult>&& Future) mutable
 	{
 		UE_LOG(LogTemp, Warning, TEXT("[FLobbiesEOS::ModifyLobbyDataImpl] Complete. Result: %s"), *LexToString(Future.Get()));
@@ -1492,25 +1633,18 @@ TFuture<TDefaultErrorResult<FLobbiesEOS::FModifyLobbyMemberDataImpl>> FLobbiesEO
 		return MakeFulfilledPromise<TDefaultErrorResult<FModifyLobbyMemberDataImpl>>(Errors::InvalidParams()).GetFuture();
 	}
 
-	TDefaultErrorResultInternal<TSharedPtr<FLobbyDetailsEOS>> DetailsResult = FLobbyDetailsEOS::CreateFromLobbyId(
-		LobbyPrerequisites.ToSharedRef(), Params.LobbyData->GetLobbyIdEOS(), Params.LocalUserId);
-	if (DetailsResult.IsError())
+	TSharedPtr<FLobbyDetailsEOS> LobbyDetails = Params.LobbyData->GetUserLobbyDetails(Params.LocalUserId);
+	if (!LobbyDetails)
 	{
-		// Todo: Errors
-		return MakeFulfilledPromise<TDefaultErrorResult<FModifyLobbyMemberDataImpl>>(Errors::Unknown(MoveTemp(DetailsResult.GetErrorValue()))).GetFuture();
+		return MakeFulfilledPromise<TDefaultErrorResult<FModifyLobbyMemberDataImpl>>(Errors::InvalidParams()).GetFuture();
 	}
 
-	UE_LOG(LogTemp, Warning, TEXT("[FLobbiesEOS::ModifyLobbyMemberDataImpl] Start. Lobby: %s, Member: %s"),
-		*Params.LobbyData->GetLobbyId(), *ToLogString(Params.LocalUserId));
-
 	TPromise<TDefaultErrorResult<FModifyLobbyMemberDataImpl>> Promise;
 	TFuture<TDefaultErrorResult<FModifyLobbyMemberDataImpl>> Future = Promise.GetFuture();
 
-	DetailsResult.GetOkValue()->ApplyLobbyMemberDataUpdates(Params.LocalUserId, MoveTemp(*Params.Changes))
+	LobbyDetails->ApplyLobbyMemberDataUpdateFromLocalChanges(Params.LocalUserId, *Params.Changes)
 	.Then([Promise = MoveTemp(Promise)](TFuture<EOS_EResult>&& Future) mutable
 	{
-		UE_LOG(LogTemp, Warning, TEXT("[FLobbiesEOS::ModifyLobbyMemberDataImpl] Complete. Result: %s"), *LexToString(Future.Get()));
-
 		// Todo: Handle "no change" better.
 		if (Future.Get() != EOS_EResult::EOS_Success && Future.Get() != EOS_EResult::EOS_NoChange)
 		{
@@ -1525,4 +1659,98 @@ TFuture<TDefaultErrorResult<FLobbiesEOS::FModifyLobbyMemberDataImpl>> FLobbiesEO
 	return Future;
 }
 
+TOnlineAsyncOpHandle<FLobbiesEOS::FProcessLobbyNotificationImpl> FLobbiesEOS::ProcessLobbyNotificationImplOp(FProcessLobbyNotificationImpl::Params&& InParams)
+{
+	TOnlineAsyncOpRef<FProcessLobbyNotificationImpl> Op = GetOp<FProcessLobbyNotificationImpl>(MoveTemp(InParams));
+	const FProcessLobbyNotificationImpl::Params& Params = Op->GetParams();
+
+	if (!Params.LobbyData.IsValid())
+	{
+		Op->SetError(Errors::InvalidParams());
+		return Op->GetHandle();
+	}
+
+	Op->Then([this](TOnlineAsyncOp<FProcessLobbyNotificationImpl>& InAsyncOp)
+	{
+		const FProcessLobbyNotificationImpl::Params& Params = InAsyncOp.GetParams();
+
+		// Notifications do not always indicate a user. Try to find a valid lobby details object to
+		// handle acquiring data snapshots.
+		TSharedPtr<FLobbyDetailsEOS> LobbyDetails = Params.LobbyData->GetActiveLobbyDetails();
+		if (!LobbyDetails.IsValid())
+		{
+			UE_LOG(LogTemp, Log, TEXT("[FLobbiesEOS::ProcessLobbyNotificationImplOp] Failed to find active lobby details to process lobby notificaions: Lobby: %s"),
+				*Params.LobbyData->GetLobbyId());
+			InAsyncOp.SetError(Errors::Unknown());
+			return MakeFulfilledPromise<TDefaultErrorResultInternal<TSharedRef<FClientLobbySnapshot>>>().GetFuture();
+		}
+
+		InAsyncOp.Data.Set<TSharedRef<FLobbyDetailsEOS>>(LobbyDetailsKeyName, LobbyDetails.ToSharedRef());
+
+		// Fetch lobby snapshot. Fetching the snapshot resolves the account ids of all lobby members in the snapshot.
+		return LobbyDetails->GetLobbySnapshot();
+	})
+	.Then([this](TOnlineAsyncOp<FProcessLobbyNotificationImpl>& InAsyncOp, TDefaultErrorResultInternal<TSharedRef<FClientLobbySnapshot>>&& LobbySnapshotResult)
+	{
+		const FProcessLobbyNotificationImpl::Params& Params = InAsyncOp.GetParams();
+		TSharedRef<FLobbyDetailsEOS> LobbyDetails = GetOpDataChecked<TSharedRef<FLobbyDetailsEOS>>(InAsyncOp, LobbyDetailsKeyName);
+
+		if (LobbySnapshotResult.IsError())
+		{
+			// Todo: errors.
+			InAsyncOp.SetError(Errors::Unknown(MoveTempIfPossible(LobbySnapshotResult.GetErrorValue())));
+			return;
+		}
+
+		// Get member snapshots.
+		TMap<FOnlineAccountIdHandle, TSharedRef<FClientLobbyMemberSnapshot>> LobbyMemberSnapshots;
+		LobbyMemberSnapshots.Reserve(Params.MutatedMembers.Num());
+		for (EOS_ProductUserId MutatedMember : Params.MutatedMembers)
+		{
+			const FOnlineAccountIdHandle MutatedMemberAccountId = FindAccountId(MutatedMember);
+			if (MutatedMemberAccountId.IsValid())
+			{
+				TDefaultErrorResultInternal<TSharedRef<FClientLobbyMemberSnapshot>> MemberSnapshotResult = LobbyDetails->GetLobbyMemberSnapshot(MutatedMemberAccountId);
+				if (MemberSnapshotResult.IsError())
+				{
+					// Todo: errors.
+					InAsyncOp.SetError(Errors::Unknown(MoveTempIfPossible(MemberSnapshotResult.GetErrorValue())));
+					return;
+				}
+				LobbyMemberSnapshots.Add(MutatedMemberAccountId, MoveTemp(MemberSnapshotResult.GetOkValue()));
+			}
+		}
+
+		// Translate leaving members from EOS_ProductUserId to FOnlineAccountIdHandle.
+		TMap<FOnlineAccountIdHandle, ELobbyMemberLeaveReason> LeavingMemberReason;
+		LeavingMemberReason.Reserve(Params.LeavingMembers.Num());
+		for (const TPair<EOS_ProductUserId, ELobbyMemberLeaveReason>& LeavingMember : Params.LeavingMembers)
+		{
+			const FOnlineAccountIdHandle LeavingMemberAccountId = FindAccountId(LeavingMember.Key);
+			if (LeavingMemberAccountId.IsValid())
+			{
+				LeavingMemberReason.Add(LeavingMemberAccountId, LeavingMember.Value);
+			}
+		}
+
+		// Apply updates and fire notifications.
+		FApplyLobbyUpdateResult Result = Params.LobbyData->GetClientLobbyData()->ApplyLobbyUpdateFromServiceSnapshot(
+			MoveTemp(*LobbySnapshotResult.GetOkValue()),
+			MoveTemp(LobbyMemberSnapshots),
+			MoveTemp(LeavingMemberReason),
+			&LobbyEvents);
+
+		// Remove active users if needed.
+		for (FOnlineAccountIdHandle LeavingMember : Result.LeavingLocalMembers)
+		{
+			RemoveActiveLobby(LeavingMember, Params.LobbyData.ToSharedRef());
+		}
+
+		InAsyncOp.SetResult(FProcessLobbyNotificationImpl::Result{});
+	})
+	.Enqueue(GetSerialQueue());
+
+	return Op->GetHandle();
+}
+
 /* UE::Online */ }
diff --git a/Engine/Plugins/Online/Experimental/OnlineServicesEOS/Source/Private/LobbiesEOSTypes.cpp b/Engine/Plugins/Online/Experimental/OnlineServicesEOS/Source/Private/LobbiesEOSTypes.cpp
--- a/Engine/Plugins/Online/Experimental/OnlineServicesEOS/Source/Private/LobbiesEOSTypes.cpp
+++ b/Engine/Plugins/Online/Experimental/OnlineServicesEOS/Source/Private/LobbiesEOSTypes.cpp
@@ -6,7 +6,7 @@
 
 namespace UE::Online {
 
-namespace detail {
+namespace Private {
 
 FString TranslateLobbyId(EOS_LobbyId EOSLobbyId)
 {
@@ -45,58 +45,562 @@ EOS_EComparisonOp TranslateSearchComparison(ELobbyComparisonOp Op)
 	}
 }
 
-class FLobbyNotificationPauseHandleImpl final : public FLobbyNotificationPauseHandle
+} // Private
+
+// Attribute translators.
+enum class EAttributeTranslationType
+{
+	ToService,
+	FromService
+};
+
+template <EAttributeTranslationType>
+class FLobbyAttributeTranslator
+{
+public:
+};
+
+template <>
+class FLobbyAttributeTranslator<EAttributeTranslationType::ToService>
+{
+public:
+	FLobbyAttributeTranslator(const TPair<FLobbyAttributeId, FLobbyVariant>& FromAttributeData);
+	FLobbyAttributeTranslator(FLobbyAttributeId FromAttributeId, const FLobbyVariant& FromAttributeData);
+
+	const EOS_Lobby_AttributeData& GetAttributeData() const { return AttributeData; }
+
+private:
+	FTCHARToUTF8 KeyConverterStorage;
+	TOptional<FTCHARToUTF8> ValueConverterStorage;
+	EOS_Lobby_AttributeData AttributeData;
+};
+
+template <>
+class FLobbyAttributeTranslator<EAttributeTranslationType::FromService>
 {
 public:
-	using FUnregisterFn = TFunction<void()>;
+	FLobbyAttributeTranslator(const EOS_Lobby_AttributeData& FromAttributeData);
+
+	const TPair<FLobbyAttributeId, FLobbyVariant>& GetAttributeData() const { return AttributeData; }
+	TPair<FLobbyAttributeId, FLobbyVariant>& GetMutableAttributeData() { return AttributeData; }
+
+private:
+	TPair<FLobbyAttributeId, FLobbyVariant> AttributeData;
+};
+
+FLobbyAttributeTranslator<EAttributeTranslationType::ToService>::FLobbyAttributeTranslator(const TPair<FLobbyAttributeId, FLobbyVariant>& FromAttributeData)
+	: FLobbyAttributeTranslator(FromAttributeData.Key, FromAttributeData.Value)
+{
+}
+
+FLobbyAttributeTranslator<EAttributeTranslationType::ToService>::FLobbyAttributeTranslator(FLobbyAttributeId FromAttributeId, const FLobbyVariant& FromAttributeData)
+	: KeyConverterStorage(*FromAttributeId.ToString())
+{
+	AttributeData.ApiVersion = EOS_LOBBY_ATTRIBUTEDATA_API_LATEST;
+	AttributeData.Key = KeyConverterStorage.Get();
+
+	if (FromAttributeData.VariantData.IsType<FString>())
+	{
+		ValueConverterStorage.Emplace(*FromAttributeData.VariantData.Get<FString>());
+		AttributeData.ValueType = EOS_ELobbyAttributeType::EOS_AT_STRING;
+		AttributeData.Value.AsUtf8 = ValueConverterStorage->Get();
+	}
+	else if (FromAttributeData.VariantData.IsType<int64>())
+	{
+		AttributeData.ValueType = EOS_ELobbyAttributeType::EOS_AT_INT64;
+		AttributeData.Value.AsInt64 = FromAttributeData.VariantData.Get<int64>();
+	}
+	else if (FromAttributeData.VariantData.IsType<double>())
+	{
+		AttributeData.ValueType = EOS_ELobbyAttributeType::EOS_AT_DOUBLE;
+		AttributeData.Value.AsDouble = FromAttributeData.VariantData.Get<double>();
+	}
+	else if (FromAttributeData.VariantData.IsType<bool>())
+	{
+		AttributeData.ValueType = EOS_ELobbyAttributeType::EOS_AT_BOOLEAN;
+		AttributeData.Value.AsBool = FromAttributeData.VariantData.Get<bool>();
+	}
+}
+
+FLobbyAttributeTranslator<EAttributeTranslationType::FromService>::FLobbyAttributeTranslator(const EOS_Lobby_AttributeData& FromAttributeData)
+{
+	FLobbyAttributeId AttributeId = Private::TranslateLobbyAttributeId(FromAttributeData.Key);
+	FLobbyVariant VariantData;
+
+	switch (FromAttributeData.ValueType)
+	{
+	case EOS_ELobbyAttributeType::EOS_AT_BOOLEAN:
+		VariantData.Set(FromAttributeData.Value.AsBool != 0);
+		break;
+
+	case EOS_ELobbyAttributeType::EOS_AT_INT64:
+		VariantData.Set(static_cast<int64>(FromAttributeData.Value.AsInt64));
+		break;
+
+	case EOS_ELobbyAttributeType::EOS_AT_DOUBLE:
+		VariantData.Set(FromAttributeData.Value.AsDouble);
+		break;
+
+	case EOS_ELobbyAttributeType::EOS_AT_STRING:
+		VariantData.Set(UTF8_TO_TCHAR(FromAttributeData.Value.AsUtf8));
+		break;
+
+	default:
+		checkNoEntry();
+		break;
+	}
+
+	AttributeData = TPair<FLobbyAttributeId, FLobbyVariant>(MoveTemp(AttributeId), MoveTemp(VariantData));
+}
+
+FClientLobbyData::FClientLobbyData(FOnlineLobbyIdHandle LobbyId)
+	: PublicData(MakeShared<FLobby>())
+{
+	PublicData->LobbyId = LobbyId;
+}
+
+TSharedPtr<const FClientLobbyMemberSnapshot> FClientLobbyData::GetMemberData(FOnlineAccountIdHandle MemberAccountId) const
+{
+	const TSharedRef<FClientLobbyMemberSnapshot>* FoundMemberData = MemberDataStorage.Find(MemberAccountId);
+	return FoundMemberData ? *FoundMemberData : TSharedPtr<const FClientLobbyMemberSnapshot>();
+}
+
+FApplyLobbyUpdateResult FClientLobbyData::ApplyLobbyUpdateFromServiceSnapshot(
+	FClientLobbySnapshot&& LobbySnapshot,
+	TMap<FOnlineAccountIdHandle, TSharedRef<FClientLobbyMemberSnapshot>>&& LobbyMemberSnapshots,
+	TMap<FOnlineAccountIdHandle, ELobbyMemberLeaveReason>&& LeaveReasons,
+	FLobbyEvents* LobbyEvents)
+{
+	FApplyLobbyUpdateResult Result;
 
-	FLobbyNotificationPauseHandleImpl(FUnregisterFn&& UnregisterFn)
-		: UnregisterFn(MoveTemp(UnregisterFn))
+	// Notify if schema changed.
+	// Schema change notification must be processed first as the schema affects how the client processes attributes.
+	if (LobbySnapshot.SchemaName != PublicData->SchemaName)
 	{
+		PublicData->SchemaName = LobbySnapshot.SchemaName;
+		if (LobbyEvents)
+		{
+			LobbyEvents->OnLobbySchemaChanged.Broadcast(FLobbySchemaChanged{PublicData});
+		}
+	}
+
+	// Handle lobby attribute changes.
+	TSet<FLobbyAttributeId> LobbyAttributeChanges = ApplyAttributeUpdateFromSnapshot(MoveTemp(LobbySnapshot.Attributes), PublicData->Attributes);
+	if (LobbyEvents && !LobbyAttributeChanges.IsEmpty())
+	{
+		LobbyEvents->OnLobbyAttributesChanged.Broadcast(FLobbyAttributesChanged{PublicData, MoveTemp(LobbyAttributeChanges)});
+	}
+
+	// Process members who left.
+	{
+		// Build list of leaving members.
+		TArray<TPair<TSharedRef<FClientLobbyMemberSnapshot>, ELobbyMemberLeaveReason>> LeavingMembers;
+		for (TPair<FOnlineAccountIdHandle, TSharedRef<FClientLobbyMemberSnapshot>>& MemberData : MemberDataStorage)
+		{
+			// Check if member exists in new snapshot.
+			if (LobbySnapshot.Members.Find(MemberData.Key) == nullptr)
+			{
+				// Member has left - check if they gave an explicit reason.
+				ELobbyMemberLeaveReason ResovedReason = ELobbyMemberLeaveReason::Disconnected;
+				if (const ELobbyMemberLeaveReason* MemberLeaveReason = LeaveReasons.Find(MemberData.Key))
+				{
+					ResovedReason = *MemberLeaveReason;
+				}
+				else
+				{
+					UE_LOG(LogTemp, Warning, TEXT("[FClientLobbyData::ApplyLobbyUpdateFromSnapshot] Member left lobby without giving reason: Lobby: %s, Member: %s"),
+						*ToLogString(PublicData->LobbyId), *ToLogString(MemberData.Key));
+				}
+
+				LeavingMembers.Emplace(MemberData.Value, ResovedReason);
+			}
+		}
+
+		// Remove member data and notify.
+		for (const TPair<TSharedRef<FClientLobbyMemberSnapshot>, ELobbyMemberLeaveReason>& MemberData : LeavingMembers)
+		{
+			MemberDataStorage.Remove(MemberData.Key->AccountId);
+			PublicData->Members.Remove(MemberData.Key->AccountId);
+
+			if (MemberData.Key->bIsLocalMember)
+			{
+				LocalMembers.Remove(MemberData.Key->AccountId);
+				Result.LeavingLocalMembers.Add(MemberData.Key->AccountId);
+			}
+
+			if (LobbyEvents)
+			{
+				LobbyEvents->OnLobbyMemberLeft.Broadcast(FLobbyMemberLeft{PublicData, MemberData.Key, MemberData.Value});
+			}
+		}
 	}
 
-	virtual ~FLobbyNotificationPauseHandleImpl()
+	// Process member joins and attribute changes.
+	// Joining members are expected to have a member data snapshot.
+	for (TPair<FOnlineAccountIdHandle, TSharedRef<FClientLobbyMemberSnapshot>>& MemberSnapshot : LobbyMemberSnapshots)
 	{
-		if (UnregisterFn)
+		if (LobbySnapshot.Members.Find(MemberSnapshot.Key) == nullptr)
 		{
-			UnregisterFn();
+			UE_LOG(LogTemp, Warning, TEXT("[FClientLobbyData::ApplyLobbyUpdateFromSnapshot] Member update ignored for unknown member: Lobby: %s, Member: %s"),
+				*ToLogString(PublicData->LobbyId), *ToLogString(MemberSnapshot.Key));
+		}
+		else
+		{
+			// Check whether this member is already in the lobby.
+			if (TSharedRef<FClientLobbyMemberSnapshot>* MemberData = MemberDataStorage.Find(MemberSnapshot.Key))
+			{
+				// Member is already in the lobby. Check if their attributes changed.
+				TSet<FLobbyAttributeId> MemberAttributeChanges = ApplyAttributeUpdateFromSnapshot(MoveTemp(MemberSnapshot.Value->Attributes), (*MemberData)->Attributes);
+				if (LobbyEvents && !MemberAttributeChanges.IsEmpty())
+				{
+					LobbyEvents->OnLobbyMemberAttributesChanged.Broadcast(FLobbyMemberAttributesChanged{PublicData, (*MemberData), MemberAttributeChanges});
+				}
+			}
+			else
+			{
+				// Member is not in the lobby, add them.
+				MemberDataStorage.Add(MemberSnapshot.Key, MemberSnapshot.Value);
+				PublicData->Members.Add(MemberSnapshot.Key, MoveTemp(MemberSnapshot.Value));
+				if (LobbyEvents)
+				{
+					LobbyEvents->OnLobbyMemberJoined.Broadcast(FLobbyMemberJoined{PublicData, MemberSnapshot.Value});
+				}
+			}
 		}
 	}
 
-private:
-	FUnregisterFn UnregisterFn;
-};
+	// Process lobby state.
+	// The lobby owner notification must be processed after member updates as the new owner may not
+	// have existed in the previous snapshot.
+	if (PublicData->OwnerAccountId != LobbySnapshot.OwnerAccountId)
+	{
+		if (TSharedRef<FClientLobbyMemberSnapshot>* MemberData = MemberDataStorage.Find(LobbySnapshot.OwnerAccountId))
+		{
+			PublicData->OwnerAccountId = LobbySnapshot.OwnerAccountId;
+			if (LobbyEvents)
+			{
+				LobbyEvents->OnLobbyLeaderChanged.Broadcast(FLobbyLeaderChanged{PublicData, (*MemberData)});
+			}
+		}
+		else
+		{
+			UE_LOG(LogTemp, Warning, TEXT("[FClientLobbyData::ApplyLobbyUpdateFromSnapshot] Lobby owner chage failed - lobby member data not found: Lobby: %s, Member: %s"),
+				*ToLogString(PublicData->LobbyId), *ToLogString(LobbySnapshot.OwnerAccountId));
+		}
+	}
+
+	// Set remaining lobby data.
+	PublicData->MaxMembers = LobbySnapshot.MaxMembers;
+	PublicData->JoinPolicy = LobbySnapshot.JoinPolicy;
+
+	// If no local members remain in the lobby process lobby removal.
+	if (LobbyEvents && LocalMembers.IsEmpty())
+	{
+		TArray<TSharedRef<FClientLobbyMemberSnapshot>> LeavingMembers;
+		MemberDataStorage.GenerateValueArray(LeavingMembers);
+
+		for (TSharedRef<FClientLobbyMemberSnapshot>& LeavingMember : LeavingMembers)
+		{
+			MemberDataStorage.Remove(LeavingMember->AccountId);
+			PublicData->Members.Remove(LeavingMember->AccountId);
+			LobbyEvents->OnLobbyMemberLeft.Broadcast(FLobbyMemberLeft{PublicData, LeavingMember, ELobbyMemberLeaveReason::Left});
+		}
+
+		LobbyEvents->OnLobbyLeft.Broadcast(FLobbyLeft{PublicData});
+	}
+
+	return Result;
+}
+
+FApplyLobbyUpdateResult FClientLobbyData::ApplyLobbyUpdateFromLocalChanges(FClientLobbyDataChanges&& Changes, FLobbyEvents& LobbyEvents)
+{
+	FApplyLobbyUpdateResult Result;
+	const bool bIsActive = !LocalMembers.IsEmpty();
+
+	// Process lobby changes.
+	{
+		if (Changes.LocalName)
+		{
+			PublicData->LocalName = *Changes.LocalName;
+		}
+
+		// Join policy.
+		if (Changes.JoinPolicy)
+		{
+			PublicData->JoinPolicy = *Changes.JoinPolicy;
+		}
+
+		// Lobby owner.
+		if (Changes.OwnerAccountId && *Changes.OwnerAccountId != PublicData->OwnerAccountId)
+		{
+			if (TSharedRef<FClientLobbyMemberSnapshot>* MemberData = MemberDataStorage.Find(*Changes.OwnerAccountId))
+			{
+				PublicData->OwnerAccountId = *Changes.OwnerAccountId;
+				if (bIsActive)
+				{
+					LobbyEvents.OnLobbyLeaderChanged.Broadcast(FLobbyLeaderChanged{PublicData, (*MemberData)});
+				}
+			}
+			else
+			{
+				UE_LOG(LogTemp, Warning, TEXT("[FClientLobbyData::ApplyLobbyUpdateFromLocalChanges] Lobby owner chage failed - lobby member data not found: Lobby: %s, Member: %s"),
+					*ToLogString(PublicData->LobbyId), *ToLogString(*Changes.OwnerAccountId));
+			}
+		}
+
+		// Lobby schema.
+		if (Changes.LobbySchema && *Changes.LobbySchema != PublicData->SchemaName)
+		{
+			PublicData->SchemaName = *Changes.LobbySchema;
+			if (bIsActive)
+			{
+				LobbyEvents.OnLobbySchemaChanged.Broadcast(FLobbySchemaChanged{PublicData});
+			}
+		}
+
+		// Apply lobby attribute changes.
+		if (!Changes.MutatedAttributes.IsEmpty() || !Changes.ClearedAttributes.IsEmpty())
+		{
+			TSet<FLobbyAttributeId> ChangedLobbyAttributes = ApplyAttributeUpdateFromChanges(
+				MoveTemp(Changes.MutatedAttributes),
+				MoveTemp(Changes.ClearedAttributes),
+				PublicData->Attributes);
+
+			if (bIsActive && !ChangedLobbyAttributes.IsEmpty())
+			{
+				LobbyEvents.OnLobbyAttributesChanged.Broadcast(FLobbyAttributesChanged{PublicData, ChangedLobbyAttributes});
+			}
+		}
+	}
+
+	// Apply member changes
+	{
+		// Adding local members will make the lobby active. When the lobby moves from inactive to
+		// active join events must be generated for the lobby and each of its existing members.
+		if (!Changes.MutatedMembers.IsEmpty())
+		{
+			if (!bIsActive)
+			{
+				LobbyEvents.OnLobbyJoined.Broadcast(FLobbyJoined{PublicData});
+
+				for (TPair<FOnlineAccountIdHandle, TSharedRef<FClientLobbyMemberSnapshot>>& MemberData : MemberDataStorage)
+				{
+					// The lobby creator will already be in the list during lobby creation due to applying a lobby snapshot.
+					if (Changes.MutatedMembers.Find(MemberData.Key) == nullptr)
+					{
+						LobbyEvents.OnLobbyMemberJoined.Broadcast(FLobbyMemberJoined{PublicData, MemberData.Value});
+					}
+				}
+			}
+
+			// Create or modify state for each mutated member and notify.
+			// The lobby creator will already be in the list during lobby creation due to applying a lobby snapshot.
+			for (TPair<FOnlineAccountIdHandle, TSharedRef<FClientLobbyMemberDataChanges>>& MutatedMemberPair : Changes.MutatedMembers)
+			{
+				if (LocalMembers.Find(MutatedMemberPair.Key))
+				{
+					// Member data is expected to be added through changes before being mutated.
+					TSharedRef<FClientLobbyMemberSnapshot>* MemberDataPtr = MemberDataStorage.Find(MutatedMemberPair.Key);
+					if (ensure(MemberDataPtr))
+					{
+						TSharedRef<FClientLobbyMemberSnapshot>& MemberData = *MemberDataPtr;
+						TSet<FLobbyAttributeId> ChangedAttributes = ApplyAttributeUpdateFromChanges(
+							MoveTemp(MutatedMemberPair.Value->MutatedAttributes),
+							MoveTemp(MutatedMemberPair.Value->ClearedAttributes),
+							MemberData->Attributes);
+
+						if (!ChangedAttributes.IsEmpty())
+						{
+							LobbyEvents.OnLobbyMemberAttributesChanged.Broadcast(FLobbyMemberAttributesChanged{PublicData, MemberData, ChangedAttributes});
+						}
+					}
+				}
+				else
+				{
+					TSharedRef<FClientLobbyMemberSnapshot> NewMember = MakeShared<FClientLobbyMemberSnapshot>();
+					NewMember->bIsLocalMember = true;
+					NewMember->AccountId = MutatedMemberPair.Key;
+					NewMember->Attributes = MoveTemp(MutatedMemberPair.Value->MutatedAttributes);
+
+					// todo:
+					//NewMember->PlatformAccountId;
+					//NewMember->PlatformDisplayName;
+
+					MemberDataStorage.Add(NewMember->AccountId, NewMember);
+					PublicData->Members.Add(NewMember->AccountId, NewMember);
+					LocalMembers.Add(NewMember->AccountId);
+					LobbyEvents.OnLobbyMemberJoined.Broadcast(FLobbyMemberJoined{PublicData, NewMember});
+				}
+			}
+		}
+
+		// Handle leaving members.
+		// When the last local member leaves, leave events will be generated for each of the
+		// remaining lobby members and for the lobby itself.
+		if (!Changes.LeavingMembers.IsEmpty())
+		{
+			for (TPair<FOnlineAccountIdHandle, ELobbyMemberLeaveReason>& LeavingMemberPair : Changes.LeavingMembers)
+			{
+				// Member data is expected to be added through changes before being mutated.
+				TSharedRef<FClientLobbyMemberSnapshot>* MemberDataPtr = MemberDataStorage.Find(LeavingMemberPair.Key);
+				if (ensure(MemberDataPtr))
+				{
+					TSharedRef<FClientLobbyMemberSnapshot> MemberData = *MemberDataPtr;
+					if (MemberData->bIsLocalMember)
+					{
+						Result.LeavingLocalMembers.Add(LeavingMemberPair.Key);
+						LocalMembers.Remove(LeavingMemberPair.Key);
+					}
+
+					MemberDataStorage.Remove(LeavingMemberPair.Key);
+					PublicData->Members.Remove(LeavingMemberPair.Key);
+					LobbyEvents.OnLobbyMemberLeft.Broadcast(FLobbyMemberLeft{PublicData, MemberData, LeavingMemberPair.Value});
+				}
+			}
+
+			// If no local members remain in the lobby process lobby removal.
+			if (LocalMembers.IsEmpty())
+			{
+				TArray<TSharedRef<FClientLobbyMemberSnapshot>> LeavingMembers;
+				MemberDataStorage.GenerateValueArray(LeavingMembers);
+
+				for (TSharedRef<FClientLobbyMemberSnapshot>& LeavingMember : LeavingMembers)
+				{
+					MemberDataStorage.Remove(LeavingMember->AccountId);
+					PublicData->Members.Remove(LeavingMember->AccountId);
+					LobbyEvents.OnLobbyMemberLeft.Broadcast(FLobbyMemberLeft{PublicData, LeavingMember, ELobbyMemberLeaveReason::Left});
+				}
+
+				LobbyEvents.OnLobbyLeft.Broadcast(FLobbyLeft{PublicData});
+			}
+		}
+	}
+
+	return Result;
+}
+
+TSet<FLobbyAttributeId> FClientLobbyData::ApplyAttributeUpdateFromSnapshot(
+	TMap<FLobbyAttributeId, FLobbyVariant>&& AttributeSnapshot,
+	TMap<FLobbyAttributeId, FLobbyVariant>& ExistingAttributes)
+{
+	TSet<FLobbyAttributeId> Result;
+
+	// Clear the attributes.
+	{
+		// Determine which attributes were removed.
+		for (const TPair<FLobbyAttributeId, FLobbyVariant>& Attribute : ExistingAttributes)
+		{
+			if (!AttributeSnapshot.Contains(Attribute.Key))
+			{
+				Result.Add(Attribute.Key);
+			}
+		}
 
-} // detail
+		// Remove them.
+		for (const FLobbyAttributeId& AttributeId : Result)
+		{
+			ExistingAttributes.Remove(AttributeId);
+		}
+	}
+
+	// Apply attribute additions and mutations.
+	{
+		for (TPair<FLobbyAttributeId, FLobbyVariant>& Attribute : AttributeSnapshot)
+		{
+			if (FLobbyVariant* ExistingAttributeData = ExistingAttributes.Find(Attribute.Key))
+			{
+				if (*ExistingAttributeData != Attribute.Value)
+				{
+					Result.Add(MoveTemp(Attribute.Key));
+					*ExistingAttributeData = MoveTemp(Attribute.Value);
+				}
+			}
+			else
+			{
+				Result.Add(Attribute.Key);
+				ExistingAttributes.Add(MoveTemp(Attribute));
+			}
+		}
+	}
+
+	return Result;
+}
+
+TSet<FLobbyAttributeId> FClientLobbyData::ApplyAttributeUpdateFromChanges(
+	TMap<FLobbyAttributeId, FLobbyVariant>&& MutatedAttributes,
+	TSet<FLobbyAttributeId>&& ClearedAttributes,
+	TMap<FLobbyAttributeId, FLobbyVariant>& ExistingAttributes)
+{
+	TSet<FLobbyAttributeId> Result;
+
+	// Apply attribute additions and mutations.
+	{
+		for (TPair<FLobbyAttributeId, FLobbyVariant>& Attribute : MutatedAttributes)
+		{
+			if (FLobbyVariant* ExistingAttributeData = ExistingAttributes.Find(Attribute.Key))
+			{
+				if (*ExistingAttributeData != Attribute.Value)
+				{
+					Result.Add(MoveTemp(Attribute.Key));
+					*ExistingAttributeData = MoveTemp(Attribute.Value);
+				}
+			}
+			else
+			{
+				Result.Add(Attribute.Key);
+				ExistingAttributes.Add(MoveTemp(Attribute));
+			}
+		}
+	}
+
+	// Clear the attributes.
+	{
+		for (FLobbyAttributeId ClearedAttributeId : ClearedAttributes)
+		{
+			if (ExistingAttributes.Remove(ClearedAttributeId) > 0)
+			{
+				Result.Add(ClearedAttributeId);
+			}
+		}
+	}
+
+	return Result;
+}
 
 const EOS_HLobbyDetails FLobbyDetailsEOS::InvalidLobbyDetailsHandle = {};
 
-TDefaultErrorResultInternal<TSharedPtr<FLobbyDetailsEOS>> FLobbyDetailsEOS::CreateFromLobbyId(const TSharedRef<FLobbyPrerequisitesEOS>& Prerequisites, EOS_LobbyId LobbyId, FOnlineAccountIdHandle MemberHandle)
+TDefaultErrorResultInternal<TSharedRef<FLobbyDetailsEOS>> FLobbyDetailsEOS::CreateFromLobbyId(
+	const TSharedRef<FLobbyPrerequisitesEOS>& Prerequisites,
+	FOnlineAccountIdHandle LocalUserId,
+	EOS_LobbyId LobbyId)
 {
 	EOS_HLobbyDetails LobbyDetailsHandle = {};
 
 	EOS_Lobby_CopyLobbyDetailsHandleOptions Options;
 	Options.ApiVersion = EOS_LOBBY_COPYLOBBYDETAILSHANDLE_API_LATEST;
 	Options.LobbyId = LobbyId;
-	Options.LocalUserId = GetProductUserIdChecked(MemberHandle);
+	Options.LocalUserId = GetProductUserIdChecked(LocalUserId);
 
 	EOS_EResult EOSResult = EOS_Lobby_CopyLobbyDetailsHandle(Prerequisites->LobbyInterfaceHandle, &Options, &LobbyDetailsHandle);
 	if (EOSResult != EOS_EResult::EOS_Success)
 	{
-		return TDefaultErrorResultInternal<TSharedPtr<FLobbyDetailsEOS>>(FromEOSError(EOSResult));
+		return TDefaultErrorResultInternal<TSharedRef<FLobbyDetailsEOS>>(FromEOSError(EOSResult));
 	}
 
-	TDefaultErrorResultInternal<TSharedPtr<FLobbyDetailsInfoEOS>> LobbyDetailsInfoResult = FLobbyDetailsInfoEOS::Create(LobbyDetailsHandle);
+	TDefaultErrorResultInternal<TSharedRef<FLobbyDetailsInfoEOS>> LobbyDetailsInfoResult = FLobbyDetailsInfoEOS::Create(LobbyDetailsHandle);
 	if (LobbyDetailsInfoResult.IsError())
 	{
 		EOS_LobbyDetails_Release(LobbyDetailsHandle);
-		return TDefaultErrorResultInternal<TSharedPtr<FLobbyDetailsEOS>>(MoveTemp(LobbyDetailsInfoResult.GetErrorValue()));
+		return TDefaultErrorResultInternal<TSharedRef<FLobbyDetailsEOS>>(MoveTemp(LobbyDetailsInfoResult.GetErrorValue()));
 	}
 
-	return TDefaultErrorResultInternal<TSharedPtr<FLobbyDetailsEOS>>(MakeShared<FLobbyDetailsEOS>(Prerequisites, LobbyDetailsInfoResult.GetOkValue().ToSharedRef(), ELobbyDetailsSource::Active, LobbyDetailsHandle));
+	return TDefaultErrorResultInternal<TSharedRef<FLobbyDetailsEOS>>(MakeShared<FLobbyDetailsEOS>(Prerequisites, LobbyDetailsInfoResult.GetOkValue(), LocalUserId, ELobbyDetailsSource::Active, LobbyDetailsHandle));
 }
 
-TDefaultErrorResultInternal<TSharedPtr<FLobbyDetailsEOS>> FLobbyDetailsEOS::CreateFromInviteId(const TSharedRef<FLobbyPrerequisitesEOS>& Prerequisites, const char* InviteId)
+TDefaultErrorResultInternal<TSharedRef<FLobbyDetailsEOS>> FLobbyDetailsEOS::CreateFromInviteId(
+	const TSharedRef<FLobbyPrerequisitesEOS>& Prerequisites,
+	FOnlineAccountIdHandle LocalUserId,
+	const char* InviteId)
 {
 	EOS_HLobbyDetails LobbyDetailsHandle = {};
 
@@ -107,20 +611,23 @@ TDefaultErrorResultInternal<TSharedPtr<FLobbyDetailsEOS>> FLobbyDetailsEOS::Crea
 	EOS_EResult EOSResult = EOS_Lobby_CopyLobbyDetailsHandleByInviteId(Prerequisites->LobbyInterfaceHandle, &Options, &LobbyDetailsHandle);
 	if (EOSResult != EOS_EResult::EOS_Success)
 	{
-		return TDefaultErrorResultInternal<TSharedPtr<FLobbyDetailsEOS>>(FromEOSError(EOSResult));
+		return TDefaultErrorResultInternal<TSharedRef<FLobbyDetailsEOS>>(FromEOSError(EOSResult));
 	}
 
-	TDefaultErrorResultInternal<TSharedPtr<FLobbyDetailsInfoEOS>> LobbyDetailsInfoResult = FLobbyDetailsInfoEOS::Create(LobbyDetailsHandle);
+	TDefaultErrorResultInternal<TSharedRef<FLobbyDetailsInfoEOS>> LobbyDetailsInfoResult = FLobbyDetailsInfoEOS::Create(LobbyDetailsHandle);
 	if (LobbyDetailsInfoResult.IsError())
 	{
 		EOS_LobbyDetails_Release(LobbyDetailsHandle);
-		return TDefaultErrorResultInternal<TSharedPtr<FLobbyDetailsEOS>>(MoveTemp(LobbyDetailsInfoResult.GetErrorValue()));
+		return TDefaultErrorResultInternal<TSharedRef<FLobbyDetailsEOS>>(MoveTemp(LobbyDetailsInfoResult.GetErrorValue()));
 	}
 
-	return TDefaultErrorResultInternal<TSharedPtr<FLobbyDetailsEOS>>(MakeShared<FLobbyDetailsEOS>(Prerequisites, LobbyDetailsInfoResult.GetOkValue().ToSharedRef(), ELobbyDetailsSource::Invite, LobbyDetailsHandle));
+	return TDefaultErrorResultInternal<TSharedRef<FLobbyDetailsEOS>>(MakeShared<FLobbyDetailsEOS>(Prerequisites, LobbyDetailsInfoResult.GetOkValue(), LocalUserId, ELobbyDetailsSource::Invite, LobbyDetailsHandle));
 }
 
-TDefaultErrorResultInternal<TSharedPtr<FLobbyDetailsEOS>> FLobbyDetailsEOS::CreateFromUiEventId(const TSharedRef<FLobbyPrerequisitesEOS>& Prerequisites, EOS_UI_EventId UiEventId)
+TDefaultErrorResultInternal<TSharedRef<FLobbyDetailsEOS>> FLobbyDetailsEOS::CreateFromUiEventId(
+	const TSharedRef<FLobbyPrerequisitesEOS>& Prerequisites,
+	FOnlineAccountIdHandle LocalUserId,
+	EOS_UI_EventId UiEventId)
 {
 	EOS_HLobbyDetails LobbyDetailsHandle = {};
 
@@ -131,20 +638,24 @@ TDefaultErrorResultInternal<TSharedPtr<FLobbyDetailsEOS>> FLobbyDetailsEOS::Crea
 	EOS_EResult EOSResult = EOS_Lobby_CopyLobbyDetailsHandleByUiEventId(Prerequisites->LobbyInterfaceHandle, &Options, &LobbyDetailsHandle);
 	if (EOSResult != EOS_EResult::EOS_Success)
 	{
-		return TDefaultErrorResultInternal<TSharedPtr<FLobbyDetailsEOS>>(FromEOSError(EOSResult));
+		return TDefaultErrorResultInternal<TSharedRef<FLobbyDetailsEOS>>(FromEOSError(EOSResult));
 	}
 
-	TDefaultErrorResultInternal<TSharedPtr<FLobbyDetailsInfoEOS>> LobbyDetailsInfoResult = FLobbyDetailsInfoEOS::Create(LobbyDetailsHandle);
+	TDefaultErrorResultInternal<TSharedRef<FLobbyDetailsInfoEOS>> LobbyDetailsInfoResult = FLobbyDetailsInfoEOS::Create(LobbyDetailsHandle);
 	if (LobbyDetailsInfoResult.IsError())
 	{
 		EOS_LobbyDetails_Release(LobbyDetailsHandle);
-		return TDefaultErrorResultInternal<TSharedPtr<FLobbyDetailsEOS>>(MoveTemp(LobbyDetailsInfoResult.GetErrorValue()));
+		return TDefaultErrorResultInternal<TSharedRef<FLobbyDetailsEOS>>(MoveTemp(LobbyDetailsInfoResult.GetErrorValue()));
 	}
 
-	return TDefaultErrorResultInternal<TSharedPtr<FLobbyDetailsEOS>>(MakeShared<FLobbyDetailsEOS>(Prerequisites, LobbyDetailsInfoResult.GetOkValue().ToSharedRef(), ELobbyDetailsSource::UiEvent, LobbyDetailsHandle));
+	return TDefaultErrorResultInternal<TSharedRef<FLobbyDetailsEOS>>(MakeShared<FLobbyDetailsEOS>(Prerequisites, LobbyDetailsInfoResult.GetOkValue(), LocalUserId, ELobbyDetailsSource::UiEvent, LobbyDetailsHandle));
 }
 
-TDefaultErrorResultInternal<TSharedPtr<FLobbyDetailsEOS>> FLobbyDetailsEOS::CreateFromSearchResult(const TSharedRef<FLobbyPrerequisitesEOS>& Prerequisites, EOS_HLobbySearch SearchHandle, uint32_t ResultIndex)
+TDefaultErrorResultInternal<TSharedRef<FLobbyDetailsEOS>> FLobbyDetailsEOS::CreateFromSearchResult(
+	const TSharedRef<FLobbyPrerequisitesEOS>& Prerequisites,
+	FOnlineAccountIdHandle LocalUserId,
+	EOS_HLobbySearch SearchHandle,
+	uint32_t ResultIndex)
 {
 	EOS_HLobbyDetails LobbyDetailsHandle = {};
 
@@ -155,17 +666,17 @@ TDefaultErrorResultInternal<TSharedPtr<FLobbyDetailsEOS>> FLobbyDetailsEOS::Crea
 	EOS_EResult EOSResult = EOS_LobbySearch_CopySearchResultByIndex(SearchHandle, &Options, &LobbyDetailsHandle);
 	if (EOSResult != EOS_EResult::EOS_Success)
 	{
-		return TDefaultErrorResultInternal<TSharedPtr<FLobbyDetailsEOS>>(FromEOSError(EOSResult));
+		return TDefaultErrorResultInternal<TSharedRef<FLobbyDetailsEOS>>(FromEOSError(EOSResult));
 	}
 
-	TDefaultErrorResultInternal<TSharedPtr<FLobbyDetailsInfoEOS>> LobbyDetailsInfoResult = FLobbyDetailsInfoEOS::Create(LobbyDetailsHandle);
+	TDefaultErrorResultInternal<TSharedRef<FLobbyDetailsInfoEOS>> LobbyDetailsInfoResult = FLobbyDetailsInfoEOS::Create(LobbyDetailsHandle);
 	if (LobbyDetailsInfoResult.IsError())
 	{
 		EOS_LobbyDetails_Release(LobbyDetailsHandle);
-		return TDefaultErrorResultInternal<TSharedPtr<FLobbyDetailsEOS>>(MoveTemp(LobbyDetailsInfoResult.GetErrorValue()));
+		return TDefaultErrorResultInternal<TSharedRef<FLobbyDetailsEOS>>(MoveTemp(LobbyDetailsInfoResult.GetErrorValue()));
 	}
 
-	return TDefaultErrorResultInternal<TSharedPtr<FLobbyDetailsEOS>>(MakeShared<FLobbyDetailsEOS>(Prerequisites, LobbyDetailsInfoResult.GetOkValue().ToSharedRef(), ELobbyDetailsSource::Search, LobbyDetailsHandle));
+	return TDefaultErrorResultInternal<TSharedRef<FLobbyDetailsEOS>>(MakeShared<FLobbyDetailsEOS>(Prerequisites, LobbyDetailsInfoResult.GetOkValue(), LocalUserId, ELobbyDetailsSource::Search, LobbyDetailsHandle));
 }
 
 FLobbyDetailsEOS::~FLobbyDetailsEOS()
@@ -173,8 +684,14 @@ FLobbyDetailsEOS::~FLobbyDetailsEOS()
 	EOS_LobbyDetails_Release(LobbyDetailsHandle);
 }
 
-TFuture<TDefaultErrorResultInternal<TSharedPtr<FLobby>>> FLobbyDetailsEOS::GetLobbyData(FOnlineAccountIdHandle LocalUserId, FOnlineLobbyIdHandle LobbyIdHandle) const
+TFuture<TDefaultErrorResultInternal<TSharedRef<FClientLobbySnapshot>>> FLobbyDetailsEOS::GetLobbySnapshot() const
 {
+	TSharedPtr<FAuthEOS> AuthInterface = Prerequisites->AuthInterface.Pin();
+	if (!AuthInterface)
+	{
+		return MakeFulfilledPromise<TDefaultErrorResultInternal<TSharedRef<FClientLobbySnapshot>>>(Errors::MissingInterface()).GetFuture();
+	}
+
 	EOS_LobbyDetails_GetMemberCountOptions GetMemberCountOptions = {};
 	GetMemberCountOptions.ApiVersion = EOS_LOBBYDETAILS_GETMEMBERCOUNT_API_LATEST;
 	const uint32_t MemberCount = EOS_LobbyDetails_GetMemberCount(LobbyDetailsHandle, &GetMemberCountOptions);
@@ -190,11 +707,12 @@ TFuture<TDefaultErrorResultInternal<TSharedPtr<FLobby>>> FLobbyDetailsEOS::GetLo
 		MemberProductUserIds->Emplace(EOS_LobbyDetails_GetMemberByIndex(LobbyDetailsHandle, &GetMemberByIndexOptions));
 	}
 
-	TPromise<TDefaultErrorResultInternal<TSharedPtr<FLobby>>> Promise;
-	TFuture<TDefaultErrorResultInternal<TSharedPtr<FLobby>>> Future = Promise.GetFuture();
+	TPromise<TDefaultErrorResultInternal<TSharedRef<FClientLobbySnapshot>>> Promise;
+	TFuture<TDefaultErrorResultInternal<TSharedRef<FClientLobbySnapshot>>> Future = Promise.GetFuture();
 
-	Prerequisites->AuthInterface->ResolveAccountIds(LocalUserId, *MemberProductUserIds)
-	.Then([StrongThis = AsShared(), Promise = MoveTemp(Promise), MemberProductUserIds, LobbyIdHandle](TFuture<TArray<FOnlineAccountIdHandle>>&& Future) mutable
+	// Resolve lobby member product user ids to FOnlineAccountIdHandle before proceeding.
+	AuthInterface->ResolveAccountIds(AssociatedLocalUser, *MemberProductUserIds)
+	.Then([StrongThis = AsShared(), Promise = MoveTemp(Promise), MemberProductUserIds](TFuture<TArray<FOnlineAccountIdHandle>>&& Future) mutable
 	{
 		const TArray<FOnlineAccountIdHandle>& ResolvedAccountIds = Future.Get();
 		if (MemberProductUserIds->Num() != ResolvedAccountIds.Num())
@@ -204,63 +722,114 @@ TFuture<TDefaultErrorResultInternal<TSharedPtr<FLobby>>> FLobbyDetailsEOS::GetLo
 			return;
 		}
 
-		EOS_LobbyDetails_GetLobbyOwnerOptions GetLobbyOwnerOptions = {};
-		GetLobbyOwnerOptions.ApiVersion = EOS_LOBBYDETAILS_GETLOBBYOWNER_API_LATEST;
-		const EOS_ProductUserId LobbyOwner = EOS_LobbyDetails_GetLobbyOwner(StrongThis->LobbyDetailsHandle, &GetLobbyOwnerOptions);
+		TSharedRef<FClientLobbySnapshot> ClientLobbySnapshot = MakeShared<FClientLobbySnapshot>();
+		ClientLobbySnapshot->MaxMembers = StrongThis->GetInfo()->Get().MaxMembers;
+		ClientLobbySnapshot->JoinPolicy = TranslateJoinPolicy(StrongThis->GetInfo()->Get().PermissionLevel);
 
-		TDefaultErrorResultInternal<TMap<FLobbyAttributeId, FLobbyVariant>> GetAttributesResult = StrongThis->GetLobbyAttributes();
-		if (GetAttributesResult.IsError())
+		// Resolve member info.
 		{
-			// Todo: Errors
-			Promise.EmplaceValue(Errors::Unknown(MoveTemp(GetAttributesResult.GetErrorValue())));
-			return;
-		}
+			EOS_LobbyDetails_GetLobbyOwnerOptions GetLobbyOwnerOptions = {};
+			GetLobbyOwnerOptions.ApiVersion = EOS_LOBBYDETAILS_GETLOBBYOWNER_API_LATEST;
+			const EOS_ProductUserId LobbyOwner = EOS_LobbyDetails_GetLobbyOwner(StrongThis->LobbyDetailsHandle, &GetLobbyOwnerOptions);
 
-		TSharedPtr<FLobby> Lobby = MakeShared<FLobby>();
-		Lobby->LobbyId = LobbyIdHandle;
-		Lobby->Attributes = GetAttributesResult.GetOkValue();
+			for (int32 MemberIndex = 0; MemberIndex < MemberProductUserIds->Num(); ++MemberIndex)
+			{
+				const EOS_ProductUserId MemberProductUserId = (*MemberProductUserIds)[MemberIndex];
+				const FOnlineAccountIdHandle ResolvedMemberAccountId = ResolvedAccountIds[MemberIndex];
 
-		for (int32 MemberIndex = 0; MemberIndex < MemberProductUserIds->Num(); ++MemberIndex)
-		{
-			const EOS_ProductUserId MemberProductUserId = (*MemberProductUserIds)[MemberIndex];
-			const FOnlineAccountIdHandle ResolvedProductUserId = ResolvedAccountIds[MemberIndex];
+				if (MemberProductUserId == LobbyOwner)
+				{
+					ClientLobbySnapshot->OwnerAccountId = ResolvedMemberAccountId;
+				}
 
-			if (MemberProductUserId == LobbyOwner)
-			{
-				Lobby->OwnerAccountId = ResolvedProductUserId;
+				ClientLobbySnapshot->Members.Add(ResolvedMemberAccountId);
 			}
+		}
+
+		// Resolve lobby attributes
+		{
+			EOS_LobbyDetails_GetAttributeCountOptions GetAttributeCountOptions = {};
+			GetAttributeCountOptions.ApiVersion = EOS_LOBBYDETAILS_GETATTRIBUTECOUNT_API_LATEST;
 
-			TDefaultErrorResultInternal<TMap<FLobbyAttributeId, FLobbyVariant>> GetMemberAttributesResult = StrongThis->GetLobbyMemberAttributes(MemberProductUserId);
-			if (GetMemberAttributesResult.IsError())
+			const uint32_t AttributeCount = EOS_LobbyDetails_GetAttributeCount(StrongThis->LobbyDetailsHandle, &GetAttributeCountOptions);
+			for (uint32_t AttributeIndex = 0; AttributeIndex < AttributeCount; ++AttributeIndex)
 			{
-				// Todo: Errors
-				Promise.EmplaceValue(Errors::Unknown(MoveTemp(GetMemberAttributesResult.GetErrorValue())));
-				return;
-			}
+				EOS_LobbyDetails_CopyAttributeByIndexOptions CopyAttributeByIndexOptions = {};
+				CopyAttributeByIndexOptions.ApiVersion = EOS_LOBBYDETAILS_COPYATTRIBUTEBYINDEX_API_LATEST;
+				CopyAttributeByIndexOptions.AttrIndex = AttributeIndex;
 
-			Lobby->Members.Emplace(StrongThis->CreateLobbyMember(ResolvedProductUserId, MoveTemp(GetMemberAttributesResult.GetOkValue())));
+				EOS_Lobby_Attribute* LobbyAttribute = nullptr;
+				ON_SCOPE_EXIT
+				{
+					EOS_Lobby_Attribute_Release(LobbyAttribute);
+				};
+
+				EOS_EResult EOSResult = EOS_LobbyDetails_CopyAttributeByIndex(StrongThis->LobbyDetailsHandle, &CopyAttributeByIndexOptions, &LobbyAttribute);
+				if (EOSResult != EOS_EResult::EOS_Success)
+				{
+					// todo: errors
+					Promise.EmplaceValue(FromEOSError(EOSResult));
+					return;
+				}
+
+				FLobbyAttributeTranslator<EAttributeTranslationType::FromService> AttributeTranslator(*LobbyAttribute->Data);
+				ClientLobbySnapshot->Attributes.Add(MoveTemp(AttributeTranslator.GetMutableAttributeData()));
+			}
 		}
 
-		Promise.EmplaceValue(MoveTemp(Lobby));
+		Promise.EmplaceValue(MoveTemp(ClientLobbySnapshot));
 	});
 
 	return Future;
 }
 
-TDefaultErrorResultInternal<TSharedPtr<FLobbyMember>> FLobbyDetailsEOS::GetLobbyMemberData(FOnlineAccountIdHandle MemberHandle) const
+TDefaultErrorResultInternal<TSharedRef<FClientLobbyMemberSnapshot>> FLobbyDetailsEOS::GetLobbyMemberSnapshot(FOnlineAccountIdHandle MemberAccountId) const
 {
-	TDefaultErrorResultInternal<TMap<FLobbyAttributeId, FLobbyVariant>> MemberAttributesResult = GetLobbyMemberAttributes(GetProductUserIdChecked(MemberHandle));
-	if (MemberAttributesResult.IsError())
-	{
-		return TDefaultErrorResultInternal<TSharedPtr<FLobbyMember>>(MoveTemp(MemberAttributesResult.GetErrorValue()));
-	}
-	else
+	EOS_ProductUserId MemberProductUserId = GetProductUserIdChecked(MemberAccountId);
+
+	TSharedRef<FClientLobbyMemberSnapshot> LobbyMemberSnapshot = MakeShared<FClientLobbyMemberSnapshot>();
+	LobbyMemberSnapshot->AccountId = MemberAccountId;
+	// Todo: 
+	//ClientMemberData->PlatformAccountId;
+	//ClientMemberData->PlatformDisplayName;
+
+	// Fetch attributes.
 	{
-		return TDefaultErrorResultInternal<TSharedPtr<FLobbyMember>>(CreateLobbyMember(MemberHandle, MoveTemp(MemberAttributesResult.GetOkValue())));
+		EOS_LobbyDetails_GetMemberAttributeCountOptions GetMemberAttributeCountOptions = {};
+		GetMemberAttributeCountOptions.ApiVersion = EOS_LOBBYDETAILS_GETMEMBERATTRIBUTECOUNT_API_LATEST;
+		GetMemberAttributeCountOptions.TargetUserId = MemberProductUserId;
+
+		const uint32_t MemberAttributeCount = EOS_LobbyDetails_GetMemberAttributeCount(LobbyDetailsHandle, &GetMemberAttributeCountOptions);
+		for (uint32_t MemberAttributeIndex = 0; MemberAttributeIndex < MemberAttributeCount; ++MemberAttributeIndex)
+		{
+			EOS_LobbyDetails_CopyMemberAttributeByIndexOptions CopyMemberAttributeByIndexOptions = {};
+			CopyMemberAttributeByIndexOptions.ApiVersion = EOS_LOBBYDETAILS_COPYMEMBERATTRIBUTEBYINDEX_API_LATEST;
+			CopyMemberAttributeByIndexOptions.TargetUserId = MemberProductUserId;
+			CopyMemberAttributeByIndexOptions.AttrIndex = MemberAttributeIndex;
+
+			EOS_Lobby_Attribute* LobbyAttribute = nullptr;
+			ON_SCOPE_EXIT
+			{
+				EOS_Lobby_Attribute_Release(LobbyAttribute);
+			};
+
+			EOS_EResult EOSResult = EOS_LobbyDetails_CopyMemberAttributeByIndex(LobbyDetailsHandle, &CopyMemberAttributeByIndexOptions, &LobbyAttribute);
+			if (EOSResult != EOS_EResult::EOS_Success)
+			{
+				return TDefaultErrorResultInternal<TSharedRef<FClientLobbyMemberSnapshot>>(FromEOSError(EOSResult));
+			}
+
+			FLobbyAttributeTranslator<EAttributeTranslationType::FromService> AttributeTranslator(*LobbyAttribute->Data);
+			LobbyMemberSnapshot->Attributes.Add(MoveTemp(AttributeTranslator.GetMutableAttributeData()));
+		}
 	}
+
+	return TDefaultErrorResultInternal<TSharedRef<FClientLobbyMemberSnapshot>>(MoveTemp(LobbyMemberSnapshot));
 }
 
-TFuture<EOS_EResult> FLobbyDetailsEOS::ApplyLobbyDataUpdates(FOnlineAccountIdHandle LocalUserId, FLobbyDataChanges Changes) const
+TFuture<EOS_EResult> FLobbyDetailsEOS::ApplyLobbyDataUpdateFromLocalChanges(
+	FOnlineAccountIdHandle LocalUserId,
+	const FClientLobbyDataChanges& Changes) const
 {
 	EOS_HLobbyModification LobbyModificationHandle = {};
 
@@ -273,7 +842,7 @@ TFuture<EOS_EResult> FLobbyDetailsEOS::ApplyLobbyDataUpdates(FOnlineAccountIdHan
 	EOS_Lobby_UpdateLobbyModificationOptions ModificationOptions = {};
 	ModificationOptions.ApiVersion = EOS_LOBBY_UPDATELOBBYMODIFICATION_API_LATEST;
 	ModificationOptions.LocalUserId = GetProductUserIdChecked(LocalUserId);
-	ModificationOptions.LobbyId = GetInfo()->GetLobbyIdEOS();
+	ModificationOptions.LobbyId = GetInfo()->Get().LobbyId;
 
 	EOS_EResult EOSResultCode = EOS_Lobby_UpdateLobbyModification(Prerequisites->LobbyInterfaceHandle, &ModificationOptions, &LobbyModificationHandle);
 	if (EOSResultCode != EOS_EResult::EOS_Success)
@@ -300,22 +869,12 @@ TFuture<EOS_EResult> FLobbyDetailsEOS::ApplyLobbyDataUpdates(FOnlineAccountIdHan
 	// Add attributes.
 	for (const TPair<FLobbyAttributeId, FLobbyVariant>& MutatedAttribute : Changes.MutatedAttributes)
 	{
-		const FTCHARToUTF8 KeyConverter(*MutatedAttribute.Key.ToString());
-		const FTCHARToUTF8 ValueConverter(*MutatedAttribute.Value);
-
-		// Todo: handle variant properly.
-		// Todo: Schema things.
-
-		EOS_Lobby_AttributeData AttributeData;
-		AttributeData.ApiVersion = EOS_LOBBY_ATTRIBUTEDATA_API_LATEST;
-		AttributeData.Key = KeyConverter.Get();
-		AttributeData.Value.AsUtf8 = ValueConverter.Get();
-		AttributeData.ValueType = EOS_ELobbyAttributeType::EOS_AT_STRING;
+		const FLobbyAttributeTranslator<EAttributeTranslationType::ToService> AttributeTranslator(MutatedAttribute);
 
 		EOS_LobbyModification_AddAttributeOptions AddAttributeOptions = {};
 		AddAttributeOptions.ApiVersion = EOS_LOBBYMODIFICATION_ADDATTRIBUTE_API_LATEST;
-		AddAttributeOptions.Attribute = &AttributeData;
-		//AddAttributeOptions.Visibility; // todo - get from schema
+		AddAttributeOptions.Attribute = &AttributeTranslator.GetAttributeData();
+		AddAttributeOptions.Visibility = EOS_ELobbyAttributeVisibility::EOS_LAT_PUBLIC; // todo - get from schema
 		EOSResultCode = EOS_LobbyModification_AddAttribute(LobbyModificationHandle, &AddAttributeOptions);
 		if (EOSResultCode != EOS_EResult::EOS_Success)
 		{
@@ -340,14 +899,14 @@ TFuture<EOS_EResult> FLobbyDetailsEOS::ApplyLobbyDataUpdates(FOnlineAccountIdHan
 		}
 	}
 
+	TPromise<EOS_EResult> Promise;
+	TFuture<EOS_EResult> Future = Promise.GetFuture();
+
 	// Apply lobby updates.
 	EOS_Lobby_UpdateLobbyOptions UpdateLobbyOptions = {};
 	UpdateLobbyOptions.ApiVersion = EOS_LOBBY_UPDATELOBBY_API_LATEST;
 	UpdateLobbyOptions.LobbyModificationHandle = LobbyModificationHandle;
 
-	TPromise<EOS_EResult> Promise;
-	TFuture<EOS_EResult> Future = Promise.GetFuture();
-
 	EOS_Async<EOS_Lobby_UpdateLobbyCallbackInfo>(EOS_Lobby_UpdateLobby, Prerequisites->LobbyInterfaceHandle, UpdateLobbyOptions)
 	.Then([Promise = MoveTemp(Promise)](TFuture<const EOS_Lobby_UpdateLobbyCallbackInfo*> Future) mutable
 	{
@@ -357,7 +916,9 @@ TFuture<EOS_EResult> FLobbyDetailsEOS::ApplyLobbyDataUpdates(FOnlineAccountIdHan
 	return Future;
 }
 
-TFuture<EOS_EResult> FLobbyDetailsEOS::ApplyLobbyMemberDataUpdates(FOnlineAccountIdHandle LocalUserId, FLobbyMemberDataChanges Changes) const
+TFuture<EOS_EResult> FLobbyDetailsEOS::ApplyLobbyMemberDataUpdateFromLocalChanges(
+	FOnlineAccountIdHandle LocalUserId,
+	const FClientLobbyMemberDataChanges& Changes) const
 {
 	EOS_HLobbyModification LobbyModificationHandle = {};
 
@@ -370,7 +931,7 @@ TFuture<EOS_EResult> FLobbyDetailsEOS::ApplyLobbyMemberDataUpdates(FOnlineAccoun
 	EOS_Lobby_UpdateLobbyModificationOptions ModificationOptions = {};
 	ModificationOptions.ApiVersion = EOS_LOBBY_UPDATELOBBYMODIFICATION_API_LATEST;
 	ModificationOptions.LocalUserId = GetProductUserIdChecked(LocalUserId);
-	ModificationOptions.LobbyId = GetInfo()->GetLobbyIdEOS();
+	ModificationOptions.LobbyId = GetInfo()->Get().LobbyId;
 
 	EOS_EResult EOSResultCode = EOS_Lobby_UpdateLobbyModification(Prerequisites->LobbyInterfaceHandle, &ModificationOptions, &LobbyModificationHandle);
 	if (EOSResultCode != EOS_EResult::EOS_Success)
@@ -382,20 +943,12 @@ TFuture<EOS_EResult> FLobbyDetailsEOS::ApplyLobbyMemberDataUpdates(FOnlineAccoun
 	// Add member attributes.
 	for (const TPair<FLobbyAttributeId, FLobbyVariant>& MutatedAttribute : Changes.MutatedAttributes)
 	{
-		const FTCHARToUTF8 KeyConverter(*MutatedAttribute.Key.ToString());
-		const FTCHARToUTF8 ValueConverter(*MutatedAttribute.Value);
-
-		// Todo: handle variant properly.
-		// Todo: Schema things.
-
-		EOS_Lobby_AttributeData AttributeData;
-		AttributeData.Key = KeyConverter.Get();
-		AttributeData.Value.AsUtf8 = ValueConverter.Get();
+		const FLobbyAttributeTranslator<EAttributeTranslationType::ToService> AttributeTranslator(MutatedAttribute);
 
 		EOS_LobbyModification_AddMemberAttributeOptions AddMemberAttributeOptions = {};
 		AddMemberAttributeOptions.ApiVersion = EOS_LOBBYMODIFICATION_ADDMEMBERATTRIBUTE_API_LATEST;
-		AddMemberAttributeOptions.Attribute = &AttributeData;
-		//AddMemberAttributeOptions.Visibility; // todo - get from schema
+		AddMemberAttributeOptions.Attribute = &AttributeTranslator.GetAttributeData();
+		AddMemberAttributeOptions.Visibility = EOS_ELobbyAttributeVisibility::EOS_LAT_PUBLIC; // todo - get from schema
 		EOSResultCode = EOS_LobbyModification_AddMemberAttribute(LobbyModificationHandle, &AddMemberAttributeOptions);
 		if (EOSResultCode != EOS_EResult::EOS_Success)
 		{
@@ -440,130 +993,18 @@ TFuture<EOS_EResult> FLobbyDetailsEOS::ApplyLobbyMemberDataUpdates(FOnlineAccoun
 FLobbyDetailsEOS::FLobbyDetailsEOS(
 	const TSharedRef<FLobbyPrerequisitesEOS>& Prerequisites,
 	const TSharedRef<FLobbyDetailsInfoEOS>& LobbyDetailsInfo,
+	FOnlineAccountIdHandle LocalUserId,
 	ELobbyDetailsSource LobbyDetailsSource,
 	EOS_HLobbyDetails LobbyDetailsHandle)
 	: Prerequisites(Prerequisites)
 	, LobbyDetailsInfo(LobbyDetailsInfo)
-	, LobbyDetailsHandle(LobbyDetailsHandle)
+	, AssociatedLocalUser(LocalUserId)
 	, LobbyDetailsSource(LobbyDetailsSource)
+	, LobbyDetailsHandle(LobbyDetailsHandle)
 {
 }
 
-TPair<FLobbyAttributeId, FLobbyVariant> FLobbyDetailsEOS::TranslateLobbyAttribute(const EOS_Lobby_Attribute& LobbyAttribute) const
-{
-	const FLobbyAttributeId AttributeId = detail::TranslateLobbyAttributeId(LobbyAttribute.Data->Key);
-	FLobbyVariant VariantData;
-
-	switch (LobbyAttribute.Data->ValueType)
-	{
-	case EOS_ELobbyAttributeType::EOS_AT_BOOLEAN:
-		// Todo:
-		break;
-
-	case EOS_ELobbyAttributeType::EOS_AT_INT64:
-		// Todo:
-		break;
-
-	case EOS_ELobbyAttributeType::EOS_AT_DOUBLE:
-		// Todo:
-		break;
-
-	case EOS_ELobbyAttributeType::EOS_AT_STRING:
-		VariantData = UTF8_TO_TCHAR(LobbyAttribute.Data->Value.AsUtf8);
-		break;
-
-	default:
-		// Todo: log
-		break;
-	}
-
-	return TPair<FLobbyAttributeId, FLobbyVariant>(AttributeId, MoveTemp(VariantData));
-}
-
-TDefaultErrorResultInternal<TMap<FLobbyAttributeId, FLobbyVariant>> FLobbyDetailsEOS::GetLobbyAttributes() const
-{
-	TMap<FLobbyAttributeId, FLobbyVariant> Attributes;
-
-	EOS_LobbyDetails_GetAttributeCountOptions GetAttributeCountOptions = {};
-	GetAttributeCountOptions.ApiVersion = EOS_LOBBYDETAILS_GETATTRIBUTECOUNT_API_LATEST;
-
-	const uint32_t AttributeCount = EOS_LobbyDetails_GetAttributeCount(LobbyDetailsHandle, &GetAttributeCountOptions);
-	for (uint32_t AttributeIndex = 0; AttributeIndex < AttributeCount; ++AttributeIndex)
-	{
-		EOS_LobbyDetails_CopyAttributeByIndexOptions CopyAttributeByIndexOptions = {};
-		CopyAttributeByIndexOptions.ApiVersion = EOS_LOBBYDETAILS_COPYATTRIBUTEBYINDEX_API_LATEST;
-		CopyAttributeByIndexOptions.AttrIndex = AttributeIndex;
-
-		EOS_Lobby_Attribute* LobbyAttribute = nullptr;
-		ON_SCOPE_EXIT
-		{
-			EOS_Lobby_Attribute_Release(LobbyAttribute);
-		};
-
-		EOS_EResult EOSResult = EOS_LobbyDetails_CopyAttributeByIndex(LobbyDetailsHandle, &CopyAttributeByIndexOptions, &LobbyAttribute);
-		if (EOSResult != EOS_EResult::EOS_Success)
-		{
-			return TDefaultErrorResultInternal<TMap<FLobbyAttributeId, FLobbyVariant>>(FromEOSError(EOSResult));
-		}
-
-		// Todo: Schema things.
-		TPair<FLobbyAttributeId, FLobbyVariant> TranslatedAttribute = TranslateLobbyAttribute(*LobbyAttribute);
-		Attributes.Emplace(TranslatedAttribute.Key, MoveTemp(TranslatedAttribute.Value));
-	}
-
-	return TDefaultErrorResultInternal<TMap<FLobbyAttributeId, FLobbyVariant>>(MoveTemp(Attributes));
-}
-
-TDefaultErrorResultInternal<TMap<FLobbyAttributeId, FLobbyVariant>> FLobbyDetailsEOS::GetLobbyMemberAttributes(EOS_ProductUserId TargetMemberProductUserId) const
-{
-	TMap<FLobbyAttributeId, FLobbyVariant> Attributes;
-
-	EOS_LobbyDetails_GetMemberAttributeCountOptions GetMemberAttributeCountOptions = {};
-	GetMemberAttributeCountOptions.ApiVersion = EOS_LOBBYDETAILS_GETMEMBERATTRIBUTECOUNT_API_LATEST;
-	GetMemberAttributeCountOptions.TargetUserId = TargetMemberProductUserId;
-
-	const uint32_t MemberAttributeCount = EOS_LobbyDetails_GetMemberAttributeCount(LobbyDetailsHandle, &GetMemberAttributeCountOptions);
-	for (uint32_t MemberAttributeIndex = 0; MemberAttributeIndex < MemberAttributeCount; ++MemberAttributeIndex)
-	{
-		EOS_LobbyDetails_CopyMemberAttributeByIndexOptions CopyMemberAttributeByIndexOptions = {};
-		CopyMemberAttributeByIndexOptions.ApiVersion = EOS_LOBBYDETAILS_COPYMEMBERATTRIBUTEBYINDEX_API_LATEST;
-		CopyMemberAttributeByIndexOptions.TargetUserId = TargetMemberProductUserId;
-		CopyMemberAttributeByIndexOptions.AttrIndex = MemberAttributeIndex;
-
-		EOS_Lobby_Attribute* LobbyAttribute = nullptr;
-		ON_SCOPE_EXIT
-		{
-			EOS_Lobby_Attribute_Release(LobbyAttribute);
-		};
-
-		EOS_EResult EOSResult = EOS_LobbyDetails_CopyMemberAttributeByIndex(LobbyDetailsHandle, &CopyMemberAttributeByIndexOptions, &LobbyAttribute);
-		if (EOSResult != EOS_EResult::EOS_Success)
-		{
-			return TDefaultErrorResultInternal<TMap<FLobbyAttributeId, FLobbyVariant>>(FromEOSError(EOSResult));
-		}
-
-		// Todo: Schema things.
-		TPair<FLobbyAttributeId, FLobbyVariant> TranslatedAttribute = TranslateLobbyAttribute(*LobbyAttribute);
-		Attributes.Emplace(TranslatedAttribute.Key, MoveTemp(TranslatedAttribute.Value));
-	}
-
-	return TDefaultErrorResultInternal<TMap<FLobbyAttributeId, FLobbyVariant>>(MoveTemp(Attributes));
-}
-
-TSharedRef<FLobbyMember> FLobbyDetailsEOS::CreateLobbyMember(FOnlineAccountIdHandle MemberHandle, TMap<FLobbyAttributeId, FLobbyVariant> Attributes) const
-{
-	TSharedRef<FLobbyMember> LobbyMember = MakeShared<FLobbyMember>();
-	LobbyMember->AccountId = MemberHandle;
-	LobbyMember->Attributes = MoveTemp(Attributes);
-
-	//Todo: 
-	//LobbyMember->PlatformAccountId;
-	//LobbyMember->PlatformDisplayName;
-
-	return LobbyMember;
-}
-
-TDefaultErrorResultInternal<TSharedPtr<FLobbyDetailsInfoEOS>> FLobbyDetailsInfoEOS::Create(EOS_HLobbyDetails LobbyDetailsHandle)
+TDefaultErrorResultInternal<TSharedRef<FLobbyDetailsInfoEOS>> FLobbyDetailsInfoEOS::Create(EOS_HLobbyDetails LobbyDetailsHandle)
 {
 	EOS_LobbyDetails_CopyInfoOptions CopyInfoOptions = {};
 	CopyInfoOptions.ApiVersion = EOS_LOBBYDETAILS_COPYINFO_API_LATEST;
@@ -572,53 +1013,26 @@ TDefaultErrorResultInternal<TSharedPtr<FLobbyDetailsInfoEOS>> FLobbyDetailsInfoE
 	EOS_EResult EOSResult = EOS_LobbyDetails_CopyInfo(LobbyDetailsHandle, &CopyInfoOptions, &LobbyDetailsInfo);
 	if (EOSResult != EOS_EResult::EOS_Success)
 	{
-		return TDefaultErrorResultInternal<TSharedPtr<FLobbyDetailsInfoEOS>>(FromEOSError(EOSResult));
+		return TDefaultErrorResultInternal<TSharedRef<FLobbyDetailsInfoEOS>>(FromEOSError(EOSResult));
 	}
 
-	TSharedRef<FLobbyDetailsInfoEOS> Result = MakeShared<FLobbyDetailsInfoEOS>();
-	Result->LobbyDetailsInfo.Reset(LobbyDetailsInfo);
-
-	// Resolve lobby info.
-
-	return TDefaultErrorResultInternal<TSharedPtr<FLobbyDetailsInfoEOS>>(MoveTemp(Result));
+	return TDefaultErrorResultInternal<TSharedRef<FLobbyDetailsInfoEOS>>(MakeShared<FLobbyDetailsInfoEOS>(FLobbyDetailsInfoPtr(LobbyDetailsInfo)));
 }
 
-FLobbyDataEOS::~FLobbyDataEOS()
+FLobbyDetailsInfoEOS::FLobbyDetailsInfoEOS(FLobbyDetailsInfoPtr&& LobbyDetailsInfo)
+	: LobbyDetailsInfo(MoveTempIfPossible(LobbyDetailsInfo))
 {
-	if (UnregisterFn)
-	{
-		UnregisterFn(LobbyIdHandle);
-	}
 }
 
-TSharedRef<FLobbyNotificationPauseHandle> FLobbyDataEOS::PauseLobbyNotifications(FOnlineAccountIdHandle LocalUserId)
+FLobbyDataEOS::~FLobbyDataEOS()
 {
-	if (TSharedRef<FLobbyNotificationPauseHandle>* ExistingPause = NotificationPauses.Find(LocalUserId))
+	if (UnregisterFn)
 	{
-		return *ExistingPause;
+		UnregisterFn(ClientLobbyData->GetPublicData().LobbyId);
 	}
-	else
-	{
-		TSharedRef<FLobbyNotificationPauseHandle> NewPause = MakeShared<detail::FLobbyNotificationPauseHandleImpl>(
-		[WeakThis = AsWeak(), LocalUserId]()
-		{
-			if (TSharedPtr<FLobbyDataEOS> StrongThis = WeakThis.Pin())
-			{
-				StrongThis->NotificationPauses.Remove(LocalUserId);
-			}
-		});
-
-		NotificationPauses.Add(LocalUserId, NewPause);
-		return NewPause;
-	}
-}
-
-bool FLobbyDataEOS::CanNotify(FOnlineAccountIdHandle LocalUserId)
-{
-	return !LocalMembers.IsEmpty() && NotificationPauses.Find(LocalUserId) == nullptr;
 }
 
-void FLobbyDataEOS::SetUserLobbyDetails(FOnlineAccountIdHandle LocalUserId, const TSharedPtr<FLobbyDetailsEOS>& LobbyDetails)
+void FLobbyDataEOS::AddUserLobbyDetails(FOnlineAccountIdHandle LocalUserId, const TSharedPtr<FLobbyDetailsEOS>& LobbyDetails)
 {
 	if (TSharedPtr<FLobbyDetailsEOS> ExistingDetails = GetUserLobbyDetails(LocalUserId))
 	{
@@ -637,51 +1051,82 @@ TSharedPtr<FLobbyDetailsEOS> FLobbyDataEOS::GetUserLobbyDetails(FOnlineAccountId
 	return Result ? *Result : TSharedPtr<FLobbyDetailsEOS>();
 }
 
+TSharedPtr<FLobbyDetailsEOS> FLobbyDataEOS::GetActiveLobbyDetails() const
+{
+	TSharedPtr<FLobbyDetailsEOS> FoundDetails;
+
+	for (const TPair<FOnlineAccountIdHandle, TSharedPtr<FLobbyDetailsEOS>>& LobbyDetails : UserLobbyDetails)
+	{
+		if (LobbyDetails.Value->GetDetailsSource() == ELobbyDetailsSource::Active)
+		{
+			FoundDetails = LobbyDetails.Value;
+			break;
+		}
+	}
+
+	return FoundDetails;
+}
+
 FLobbyDataEOS::FLobbyDataEOS(
-	FOnlineLobbyIdHandle LobbyIdHandle,
-	const TSharedRef<FLobby>& LobbyImpl,
+	const TSharedRef<FClientLobbyData>& ClientLobbyData,
 	const TSharedRef<FLobbyDetailsInfoEOS>& LobbyDetailsInfo,
 	FUnregisterFn UnregisterFn)
-	: LobbyIdHandle(LobbyIdHandle)
-	, LobbyImpl(LobbyImpl)
+	: ClientLobbyData(ClientLobbyData)
 	, LobbyDetailsInfo(LobbyDetailsInfo)
 	, UnregisterFn(MoveTemp(UnregisterFn))
-	, LobbyId(detail::TranslateLobbyId(LobbyDetailsInfo->GetLobbyIdEOS()))
+	, LobbyId(Private::TranslateLobbyId(LobbyDetailsInfo->Get().LobbyId))
 {
 }
 
-TFuture<TDefaultErrorResultInternal<TSharedPtr<FLobbyDataEOS>>> FLobbyDataEOS::Create(
-	FOnlineAccountIdHandle LocalUserId,
+TFuture<TDefaultErrorResultInternal<TSharedRef<FLobbyDataEOS>>> FLobbyDataEOS::Create(
 	FOnlineLobbyIdHandle LobbyIdHandle,
 	const TSharedRef<FLobbyDetailsEOS>& LobbyDetails,
 	FUnregisterFn UnregisterFn)
 {
-	TPromise<TDefaultErrorResultInternal<TSharedPtr<FLobbyDataEOS>>> Promise;
-	TFuture<TDefaultErrorResultInternal<TSharedPtr<FLobbyDataEOS>>> Future = Promise.GetFuture();
+	TPromise<TDefaultErrorResultInternal<TSharedRef<FLobbyDataEOS>>> Promise;
+	TFuture<TDefaultErrorResultInternal<TSharedRef<FLobbyDataEOS>>> Future = Promise.GetFuture();
 	
-	LobbyDetails->GetLobbyData(LocalUserId, LobbyIdHandle)
+	LobbyDetails->GetLobbySnapshot()
 	.Then(
 	[
 		Promise = MoveTemp(Promise),
-		LocalUserId,
 		LobbyIdHandle,
 		LobbyDetails,
 		UnregisterFn = MoveTemp(UnregisterFn)
 	]
-	(TFuture<TDefaultErrorResultInternal<TSharedPtr<FLobby>>>&& Future) mutable
+	(TFuture<TDefaultErrorResultInternal<TSharedRef<FClientLobbySnapshot>>>&& Future) mutable
 	{
 		if (Future.Get().IsError())
 		{
+			// todo: errors.
 			Promise.EmplaceValue(MoveTemp(Future.Get().GetErrorValue()));
+			return;
 		}
-		else
+
+		TSharedRef<FClientLobbySnapshot> LobbySnapshot = Future.Get().GetOkValue();
+		TSharedRef<FClientLobbyData> LobbyData = MakeShared<FClientLobbyData>(LobbyIdHandle);
+
+		// Fetch member data and apply them to the lobby.
+		TMap<FOnlineAccountIdHandle, TSharedRef<FClientLobbyMemberSnapshot>> MemberSnapshots;
+		for (FOnlineAccountIdHandle MemberAccountId : LobbySnapshot->Members)
 		{
-			Promise.EmplaceValue(MakeShared<FLobbyDataEOS>(
-				LobbyIdHandle,
-				Future.Get().GetOkValue().ToSharedRef(),
-				LobbyDetails->GetInfo(),
-				MoveTemp(UnregisterFn)));
+			TDefaultErrorResultInternal<TSharedRef<FClientLobbyMemberSnapshot>> LobbyMemberSnapshotResult = LobbyDetails->GetLobbyMemberSnapshot(MemberAccountId);
+			if (LobbyMemberSnapshotResult.IsError())
+			{
+				// todo: errors.
+				Promise.EmplaceValue(MoveTemp(LobbyMemberSnapshotResult.GetErrorValue()));
+				return;
+			}
+
+			MemberSnapshots.Emplace(MemberAccountId, MoveTemp(LobbyMemberSnapshotResult.GetOkValue()));
 		}
+
+		LobbyData->ApplyLobbyUpdateFromServiceSnapshot(MoveTemp(*LobbySnapshot), MoveTemp(MemberSnapshots));
+
+		Promise.EmplaceValue(MakeShared<FLobbyDataEOS>(
+			LobbyData,
+			LobbyDetails->GetInfo(),
+			MoveTemp(UnregisterFn)));
 	});
 
 	return Future;
@@ -694,7 +1139,7 @@ FLobbyDataRegistryEOS::FLobbyDataRegistryEOS(const TSharedRef<FLobbyPrerequisite
 
 TSharedPtr<FLobbyDataEOS> FLobbyDataRegistryEOS::Find(EOS_LobbyId EOSLobbyId) const
 {
-	const TWeakPtr<FLobbyDataEOS>* Result = LobbyIdIndex.Find(detail::TranslateLobbyId(EOSLobbyId));
+	const TWeakPtr<FLobbyDataEOS>* Result = LobbyIdIndex.Find(Private::TranslateLobbyId(EOSLobbyId));
 	return Result ? Result->Pin() : TSharedPtr<FLobbyDataEOS>();
 }
 
@@ -704,52 +1149,30 @@ TSharedPtr<FLobbyDataEOS> FLobbyDataRegistryEOS::Find(FOnlineLobbyIdHandle Lobby
 	return Result ? Result->Pin() : TSharedPtr<FLobbyDataEOS>();
 }
 
-TFuture<TDefaultErrorResultInternal<TSharedPtr<FLobbyDataEOS>>> FLobbyDataRegistryEOS::FindOrCreateFromLobbyId(FOnlineAccountIdHandle LocalUserId, EOS_LobbyId EOSLobbyId)
+TFuture<TDefaultErrorResultInternal<TSharedRef<FLobbyDataEOS>>> FLobbyDataRegistryEOS::FindOrCreateFromLobbyDetails(FOnlineAccountIdHandle LocalUserId, const TSharedRef<FLobbyDetailsEOS>& LobbyDetails)
 {
-	TDefaultErrorResultInternal<TSharedPtr<FLobbyDetailsEOS>> LobbyDetailsResult = FLobbyDetailsEOS::CreateFromLobbyId(Prerequisites, EOSLobbyId, LocalUserId);
-	if (LobbyDetailsResult.IsError())
+	if (TSharedPtr<FLobbyDataEOS> FindResult = Find(LobbyDetails->GetInfo()->Get().LobbyId))
 	{
-		return MakeFulfilledPromise<TDefaultErrorResultInternal<TSharedPtr<FLobbyDataEOS>>>(MoveTemp(LobbyDetailsResult.GetErrorValue())).GetFuture();
+		FindResult->AddUserLobbyDetails(LocalUserId, LobbyDetails);
+		return MakeFulfilledPromise<TDefaultErrorResultInternal<TSharedRef<FLobbyDataEOS>>>(FindResult.ToSharedRef()).GetFuture();
 	}
 
-	return FindOrCreateFromLobbyDetails(LocalUserId, LobbyDetailsResult.GetOkValue().ToSharedRef());
-}
-
-TFuture<TDefaultErrorResultInternal<TSharedPtr<FLobbyDataEOS>>> FLobbyDataRegistryEOS::FindOrCreateFromUiEventId(FOnlineAccountIdHandle LocalUserId, EOS_UI_EventId UiEventId)
-{
-	TDefaultErrorResultInternal<TSharedPtr<FLobbyDetailsEOS>> LobbyDetailsResult = FLobbyDetailsEOS::CreateFromUiEventId(Prerequisites, UiEventId);
-	if (LobbyDetailsResult.IsError())
-	{
-		return MakeFulfilledPromise<TDefaultErrorResultInternal<TSharedPtr<FLobbyDataEOS>>>(MoveTemp(LobbyDetailsResult.GetErrorValue())).GetFuture();
-	}
-
-	return FindOrCreateFromLobbyDetails(LocalUserId, LobbyDetailsResult.GetOkValue().ToSharedRef());
-}
-
-TFuture<TDefaultErrorResultInternal<TSharedPtr<FLobbyDataEOS>>> FLobbyDataRegistryEOS::FindOrCreateFromLobbyDetails(FOnlineAccountIdHandle LocalUserId, const TSharedRef<FLobbyDetailsEOS>& LobbyDetails)
-{
-	if (TSharedPtr<FLobbyDataEOS> FindResult = Find(LobbyDetails->GetInfo()->GetLobbyIdEOS()))
-	{
-		FindResult->SetUserLobbyDetails(LocalUserId, LobbyDetails);
-		return MakeFulfilledPromise<TDefaultErrorResultInternal<TSharedPtr<FLobbyDataEOS>>>(FindResult).GetFuture();
-	}
-
-	TPromise<TDefaultErrorResultInternal<TSharedPtr<FLobbyDataEOS>>> Promise;
-	TFuture<TDefaultErrorResultInternal<TSharedPtr<FLobbyDataEOS>>> Future = Promise.GetFuture();
+	TPromise<TDefaultErrorResultInternal<TSharedRef<FLobbyDataEOS>>> Promise;
+	TFuture<TDefaultErrorResultInternal<TSharedRef<FLobbyDataEOS>>> Future = Promise.GetFuture();
 
 	const FOnlineLobbyIdHandle LobbyId = FOnlineLobbyIdHandle(EOnlineServices::Epic, NextHandleIndex++);
-	FLobbyDataEOS::Create(LocalUserId, LobbyId, LobbyDetails, MakeUnregisterFn())
-	.Then([WeakThis = AsWeak(), Promise = MoveTemp(Promise), LocalUserId, LobbyDetails](TFuture<TDefaultErrorResultInternal<TSharedPtr<FLobbyDataEOS>>>&& Future) mutable
+	FLobbyDataEOS::Create(LobbyId, LobbyDetails, MakeUnregisterFn())
+	.Then([WeakThis = AsWeak(), Promise = MoveTemp(Promise), LocalUserId, LobbyDetails](TFuture<TDefaultErrorResultInternal<TSharedRef<FLobbyDataEOS>>>&& Future) mutable
 	{
 		if (TSharedPtr<FLobbyDataRegistryEOS> StrongThis = WeakThis.Pin())
 		{
 			if (Future.Get().IsOk())
 			{
-				StrongThis->Register(Future.Get().GetOkValue().ToSharedRef());
+				StrongThis->Register(Future.Get().GetOkValue());
 			}
 		}
 
-		Future.Get().GetOkValue()->SetUserLobbyDetails(LocalUserId, LobbyDetails);
+		Future.Get().GetOkValue()->AddUserLobbyDetails(LocalUserId, LobbyDetails);
 		Promise.EmplaceValue(MoveTempIfPossible(Future.Get()));
 	});
 
@@ -782,27 +1205,27 @@ FLobbyDataEOS::FUnregisterFn FLobbyDataRegistryEOS::MakeUnregisterFn()
 	};
 }
 
-TFuture<TDefaultErrorResultInternal<TSharedPtr<FLobbyInviteDataEOS>>> FLobbyInviteDataEOS::CreateFromInviteId(
+TFuture<TDefaultErrorResultInternal<TSharedRef<FLobbyInviteDataEOS>>> FLobbyInviteDataEOS::CreateFromInviteId(
 	const TSharedRef<FLobbyPrerequisitesEOS>& Prerequisites,
 	const TSharedRef<FLobbyDataRegistryEOS>& LobbyDataRegistry,
 	FOnlineAccountIdHandle LocalUserId,
 	const char* InviteIdEOS,
 	EOS_ProductUserId Sender)
 {
-	TDefaultErrorResultInternal<TSharedPtr<FLobbyDetailsEOS>> LobbyDetailsResult = FLobbyDetailsEOS::CreateFromInviteId(Prerequisites, InviteIdEOS);
+	TDefaultErrorResultInternal<TSharedRef<FLobbyDetailsEOS>> LobbyDetailsResult = FLobbyDetailsEOS::CreateFromInviteId(Prerequisites, LocalUserId, InviteIdEOS);
 	if (LobbyDetailsResult.IsError())
 	{
-		return MakeFulfilledPromise<TDefaultErrorResultInternal<TSharedPtr<FLobbyInviteDataEOS>>>(MoveTemp(LobbyDetailsResult.GetErrorValue())).GetFuture();
+		return MakeFulfilledPromise<TDefaultErrorResultInternal<TSharedRef<FLobbyInviteDataEOS>>>(MoveTemp(LobbyDetailsResult.GetErrorValue())).GetFuture();
 	}
 
-	TSharedRef<FLobbyDetailsEOS> LobbyDetails = LobbyDetailsResult.GetOkValue().ToSharedRef();
-	TPromise<TDefaultErrorResultInternal<TSharedPtr<FLobbyInviteDataEOS>>> Promise;
-	TFuture<TDefaultErrorResultInternal<TSharedPtr<FLobbyInviteDataEOS>>> Future = Promise.GetFuture();
+	TSharedRef<FLobbyDetailsEOS> LobbyDetails = LobbyDetailsResult.GetOkValue();
+	TPromise<TDefaultErrorResultInternal<TSharedRef<FLobbyInviteDataEOS>>> Promise;
+	TFuture<TDefaultErrorResultInternal<TSharedRef<FLobbyInviteDataEOS>>> Future = Promise.GetFuture();
 
 	// Search for existing lobby data so that the LobbyIdHandle will match.
 	TSharedRef<FLobbyInviteIdEOS> InviteId = MakeShared<FLobbyInviteIdEOS>(InviteIdEOS);
 	LobbyDataRegistry->FindOrCreateFromLobbyDetails(LocalUserId, LobbyDetails)
-	.Then([Promise = MoveTemp(Promise), InviteId, LocalUserId, Sender, LobbyDetails](TFuture<TDefaultErrorResultInternal<TSharedPtr<FLobbyDataEOS>>>&& Future) mutable
+	.Then([Promise = MoveTemp(Promise), InviteId, LocalUserId, Sender, LobbyDetails](TFuture<TDefaultErrorResultInternal<TSharedRef<FLobbyDataEOS>>>&& Future) mutable
 	{
 		if (Future.Get().IsError())
 		{
@@ -815,11 +1238,11 @@ TFuture<TDefaultErrorResultInternal<TSharedPtr<FLobbyInviteDataEOS>>> FLobbyInvi
 		if (!SenderUserId.IsValid())
 		{
 			// Todo: Errors.
-			Promise.EmplaceValue(Errors::Unknown());
+			Promise.EmplaceValue(MoveTemp(Future.Get().GetErrorValue()));
 			return;
 		}
 
-		Promise.EmplaceValue(MakeShared<FLobbyInviteDataEOS>(InviteId, LocalUserId, SenderUserId, LobbyDetails, Future.Get().GetOkValue().ToSharedRef()));
+		Promise.EmplaceValue(MakeShared<FLobbyInviteDataEOS>(InviteId, LocalUserId, SenderUserId, LobbyDetails, Future.Get().GetOkValue()));
 	});
 
 	return Future;
@@ -836,11 +1259,11 @@ FLobbyInviteDataEOS::FLobbyInviteDataEOS(
 	, Sender(Sender)
 	, LobbyDetails(LobbyDetails)
 	, LobbyData(LobbyData)
-	, InviteId(detail::TranslateLobbyInviteId(InviteIdEOS->Get()))
+	, InviteId(Private::TranslateLobbyInviteId(InviteIdEOS->Get()))
 {
 }
 
-TFuture<TDefaultErrorResultInternal<TSharedPtr<FLobbySearchEOS>>> FLobbySearchEOS::Create(
+TFuture<TDefaultErrorResultInternal<TSharedRef<FLobbySearchEOS>>> FLobbySearchEOS::Create(
 	const TSharedRef<FLobbyPrerequisitesEOS>& Prerequisites,
 	const TSharedRef<FLobbyDataRegistryEOS>& LobbyRegistry,
 	const FLobbySearchParameters& Params)
@@ -855,7 +1278,7 @@ TFuture<TDefaultErrorResultInternal<TSharedPtr<FLobbySearchEOS>>> FLobbySearchEO
 	if (EOSResult != EOS_EResult::EOS_Success)
 	{
 		// todo: errors
-		return MakeFulfilledPromise<TDefaultErrorResultInternal<TSharedPtr<FLobbySearchEOS>>>(FromEOSError(EOSResult)).GetFuture();
+		return MakeFulfilledPromise<TDefaultErrorResultInternal<TSharedRef<FLobbySearchEOS>>>(FromEOSError(EOSResult)).GetFuture();
 	}
 
 	if (Params.LobbyId)
@@ -863,7 +1286,7 @@ TFuture<TDefaultErrorResultInternal<TSharedPtr<FLobbySearchEOS>>> FLobbySearchEO
 		TSharedPtr<FLobbyDataEOS> LobbyData = LobbyRegistry->Find(*Params.LobbyId);
 		if (!LobbyData)
 		{
-			return MakeFulfilledPromise<TDefaultErrorResultInternal<TSharedPtr<FLobbySearchEOS>>>(Errors::InvalidParams()).GetFuture();
+			return MakeFulfilledPromise<TDefaultErrorResultInternal<TSharedRef<FLobbySearchEOS>>>(Errors::InvalidParams()).GetFuture();
 		}
 
 		EOS_LobbySearch_SetLobbyIdOptions SetLobbyIdOptions = {};
@@ -874,7 +1297,7 @@ TFuture<TDefaultErrorResultInternal<TSharedPtr<FLobbySearchEOS>>> FLobbySearchEO
 		if (EOSResult != EOS_EResult::EOS_Success)
 		{
 			// todo: errors
-			return MakeFulfilledPromise<TDefaultErrorResultInternal<TSharedPtr<FLobbySearchEOS>>>(FromEOSError(EOSResult)).GetFuture();
+			return MakeFulfilledPromise<TDefaultErrorResultInternal<TSharedRef<FLobbySearchEOS>>>(FromEOSError(EOSResult)).GetFuture();
 		}
 	}
 
@@ -888,38 +1311,33 @@ TFuture<TDefaultErrorResultInternal<TSharedPtr<FLobbySearchEOS>>> FLobbySearchEO
 		if (EOSResult != EOS_EResult::EOS_Success)
 		{
 			// todo: errors
-			return MakeFulfilledPromise<TDefaultErrorResultInternal<TSharedPtr<FLobbySearchEOS>>>(FromEOSError(EOSResult)).GetFuture();
+			return MakeFulfilledPromise<TDefaultErrorResultInternal<TSharedRef<FLobbySearchEOS>>>(FromEOSError(EOSResult)).GetFuture();
 		}
 	}
 
 	for (const FFindLobbySearchFilter& Filter :  Params.Filters)
 	{
-		FTCHARToUTF8 Key(*Filter.AttributeName.ToString());
-		FTCHARToUTF8 Value(*Filter.ComparisonValue);
+		const FLobbyAttributeTranslator<EAttributeTranslationType::ToService> AttributeTranslator(
+			Filter.AttributeName, Filter.ComparisonValue);
 
 		EOS_Lobby_AttributeData AttributeData;
 		AttributeData.ApiVersion = EOS_LOBBY_ATTRIBUTEDATA_API_LATEST;
 
-		// todo: support other attribute types.
-		AttributeData.Key = Key.Get();
-		AttributeData.Value.AsUtf8 = Value.Get();
-		AttributeData.ValueType = EOS_ELobbyAttributeType::EOS_AT_STRING;
-
 		EOS_LobbySearch_SetParameterOptions SetParameterOptions = {};
 		SetParameterOptions.ApiVersion = EOS_LOBBYSEARCH_SETTARGETUSERID_API_LATEST;
-		SetParameterOptions.Parameter = &AttributeData;
-		SetParameterOptions.ComparisonOp = detail::TranslateSearchComparison(Filter.ComparisonOp);
+		SetParameterOptions.Parameter = &AttributeTranslator.GetAttributeData();
+		SetParameterOptions.ComparisonOp = Private::TranslateSearchComparison(Filter.ComparisonOp);
 
 		EOSResult = EOS_LobbySearch_SetParameter(SearchHandle->Get(), &SetParameterOptions);
 		if (EOSResult != EOS_EResult::EOS_Success)
 		{
 			// todo: errors
-			return MakeFulfilledPromise<TDefaultErrorResultInternal<TSharedPtr<FLobbySearchEOS>>>(FromEOSError(EOSResult)).GetFuture();
+			return MakeFulfilledPromise<TDefaultErrorResultInternal<TSharedRef<FLobbySearchEOS>>>(FromEOSError(EOSResult)).GetFuture();
 		}
 	}
 
-	TPromise<TDefaultErrorResultInternal<TSharedPtr<FLobbySearchEOS>>> Promise;
-	TFuture<TDefaultErrorResultInternal<TSharedPtr<FLobbySearchEOS>>> Future = Promise.GetFuture();
+	TPromise<TDefaultErrorResultInternal<TSharedRef<FLobbySearchEOS>>> Promise;
+	TFuture<TDefaultErrorResultInternal<TSharedRef<FLobbySearchEOS>>> Future = Promise.GetFuture();
 
 	EOS_LobbySearch_FindOptions FindOptions = {};
 	FindOptions.ApiVersion = EOS_LOBBYSEARCH_FIND_API_LATEST;
@@ -936,7 +1354,7 @@ TFuture<TDefaultErrorResultInternal<TSharedPtr<FLobbySearchEOS>>> FLobbySearchEO
 			return;
 		}
 
-		TArray<TFuture<TDefaultErrorResultInternal<TSharedPtr<FLobbyDataEOS>>>> ResolvedLobbyDetails;
+		TArray<TFuture<TDefaultErrorResultInternal<TSharedRef<FLobbyDataEOS>>>> ResolvedLobbyDetails;
 
 		EOS_LobbySearch_GetSearchResultCountOptions GetSearchResultCountOptions = {};
 		GetSearchResultCountOptions.ApiVersion = EOS_LOBBYSEARCH_GETSEARCHRESULTCOUNT_API_LATEST;
@@ -944,8 +1362,8 @@ TFuture<TDefaultErrorResultInternal<TSharedPtr<FLobbySearchEOS>>> FLobbySearchEO
 
 		for (uint32_t SearchResultIndex = 0; SearchResultIndex < NumSearchResults; ++SearchResultIndex)
 		{
-			TDefaultErrorResultInternal<TSharedPtr<FLobbyDetailsEOS>> Result = FLobbyDetailsEOS::CreateFromSearchResult(
-				Prerequisites, SearchHandle->Get(), SearchResultIndex);
+			TDefaultErrorResultInternal<TSharedRef<FLobbyDetailsEOS>> Result = FLobbyDetailsEOS::CreateFromSearchResult(
+				Prerequisites, LocalUserId, SearchHandle->Get(), SearchResultIndex);
 			if (Result.IsError())
 			{
 				// todo: errors
@@ -953,24 +1371,24 @@ TFuture<TDefaultErrorResultInternal<TSharedPtr<FLobbySearchEOS>>> FLobbySearchEO
 				return;
 			}
 
-			TPromise<TDefaultErrorResultInternal<TSharedPtr<FLobbyDataEOS>>> ResolveLobbyDetailsPromise;
+			TPromise<TDefaultErrorResultInternal<TSharedRef<FLobbyDataEOS>>> ResolveLobbyDetailsPromise;
 			ResolvedLobbyDetails.Add(ResolveLobbyDetailsPromise.GetFuture());
 
-			LobbyRegistry->FindOrCreateFromLobbyDetails(LocalUserId, Result.GetOkValue().ToSharedRef())
-			.Then([ResolveLobbyDetailsPromise = MoveTemp(ResolveLobbyDetailsPromise)](TFuture<TDefaultErrorResultInternal<TSharedPtr<FLobbyDataEOS>>>&& Future) mutable
+			LobbyRegistry->FindOrCreateFromLobbyDetails(LocalUserId, Result.GetOkValue())
+			.Then([ResolveLobbyDetailsPromise = MoveTemp(ResolveLobbyDetailsPromise)](TFuture<TDefaultErrorResultInternal<TSharedRef<FLobbyDataEOS>>>&& Future) mutable
 			{
 				ResolveLobbyDetailsPromise.EmplaceValue(MoveTempIfPossible(Future.Get()));
 			});
 		}
 
 		WhenAll(MoveTemp(ResolvedLobbyDetails))
-		.Then([Promise = MoveTemp(Promise), SearchHandle](TFuture<TArray<TDefaultErrorResultInternal<TSharedPtr<FLobbyDataEOS>>>> && Future) mutable
+		.Then([Promise = MoveTemp(Promise), SearchHandle](TFuture<TArray<TDefaultErrorResultInternal<TSharedRef<FLobbyDataEOS>>>>&& Future) mutable
 		{
-			TArray<TDefaultErrorResultInternal<TSharedPtr<FLobbyDataEOS>>> Results(Future.Get());
-			TArray<TSharedPtr<FLobbyDataEOS>> ResolvedResults;
+			TArray<TDefaultErrorResultInternal<TSharedRef<FLobbyDataEOS>>> Results(Future.Get());
+			TArray<TSharedRef<FLobbyDataEOS>> ResolvedResults;
 			ResolvedResults.Reserve(Results.Num());
 
-			for (TDefaultErrorResultInternal<TSharedPtr<FLobbyDataEOS>>& Result : Results)
+			for (TDefaultErrorResultInternal<TSharedRef<FLobbyDataEOS>>& Result : Results)
 			{
 				if (Result.IsError())
 				{
@@ -994,20 +1412,20 @@ TArray<TSharedRef<const FLobby>> FLobbySearchEOS::GetLobbyResults() const
 	TArray<TSharedRef<const FLobby>> Result;
 	Result.Reserve(Lobbies.Num());
 
-	for (const TSharedPtr<FLobbyDataEOS>& LobbyData : Lobbies)
+	for (const TSharedRef<FLobbyDataEOS>& LobbyData : Lobbies)
 	{
-		Result.Add(LobbyData->GetLobbyImpl());
+		Result.Add(LobbyData->GetClientLobbyData()->GetPublicDataPtr());
 	}
 
 	return Result;
 }
 
-const TArray<TSharedPtr<FLobbyDataEOS>>& FLobbySearchEOS::GetLobbyData()
+const TArray<TSharedRef<FLobbyDataEOS>>& FLobbySearchEOS::GetLobbyData()
 {
 	return Lobbies;
 }
 
-FLobbySearchEOS::FLobbySearchEOS(const TSharedRef<FSearchHandle>& SearchHandle, TArray<TSharedPtr<FLobbyDataEOS>>&& Lobbies)
+FLobbySearchEOS::FLobbySearchEOS(const TSharedRef<FSearchHandle>& SearchHandle, TArray<TSharedRef<FLobbyDataEOS>>&& Lobbies)
 	: SearchHandle(SearchHandle)
 	, Lobbies(Lobbies)
 {
diff --git a/Engine/Plugins/Online/Experimental/OnlineServicesEOS/Source/Private/LobbiesEOSTypes.h b/Engine/Plugins/Online/Experimental/OnlineServicesEOS/Source/Private/LobbiesEOSTypes.h
--- a/Engine/Plugins/Online/Experimental/OnlineServicesEOS/Source/Private/LobbiesEOSTypes.h
+++ b/Engine/Plugins/Online/Experimental/OnlineServicesEOS/Source/Private/LobbiesEOSTypes.h
@@ -20,7 +20,9 @@ class FLobbyDataRegistryEOS;
 class FLobbyDetailsInfoEOS;
 class FLobbyInviteDataEOSRegistry;
 class FLobbySchema;
+struct FLobbyEvents;
 class FAuthEOS;
+struct FClientLobbyMemberSnapshot;
 
 inline EOS_ELobbyPermissionLevel TranslateJoinPolicy(ELobbyJoinPolicy JoinPolicy)
 {
@@ -33,95 +35,273 @@ inline EOS_ELobbyPermissionLevel TranslateJoinPolicy(ELobbyJoinPolicy JoinPolicy
 	};
 }
 
-using FLobbySearchParameters = FFindLobby::Params;
-
-struct FLobbyDataChanges
+inline ELobbyJoinPolicy TranslateJoinPolicy(EOS_ELobbyPermissionLevel JoinPolicy)
 {
-	TOptional<ELobbyJoinPolicy> JoinPolicy;
+	switch (JoinPolicy)
+	{
+	case EOS_ELobbyPermissionLevel::EOS_LPL_PUBLICADVERTISED:	return ELobbyJoinPolicy::PublicAdvertised;
+	case EOS_ELobbyPermissionLevel::EOS_LPL_JOINVIAPRESENCE:	return ELobbyJoinPolicy::PublicNotAdvertised;
+	default:									checkNoEntry(); // Intentional fallthrough
+	case EOS_ELobbyPermissionLevel::EOS_LPL_INVITEONLY:	return ELobbyJoinPolicy::InvitationOnly;
+	};
+}
 
-	// New or changed attributes.
+using FLobbySearchParameters = FFindLobbies::Params;
+
+/**
+ * Local changes to a member to be applied to an existing lobby member snapshot.
+ */
+struct FClientLobbyMemberDataChanges
+{
+	/** New or changed attributes. */
 	TMap<FLobbyAttributeId, FLobbyVariant> MutatedAttributes;
 
-	// Attributes to be cleared.
+	/** Attributes to be cleared. */
 	TSet<FLobbyAttributeId> ClearedAttributes;
 };
 
-struct FLobbyMemberDataChanges
+/**
+ * Local changes to a lobby to be applied to an existing lobby snapshot.
+ */
+struct FClientLobbyDataChanges
 {
-	// New or changed attributes.
+	/** Local name for the lobby. */
+	TOptional<FName> LocalName;
+
+	/** Setting for new join policy. */
+	TOptional<ELobbyJoinPolicy> JoinPolicy;
+
+	/** Setting for lobby ownership change. */
+	TOptional<FOnlineAccountIdHandle> OwnerAccountId;
+
+	/** Setting for lobby schema change. */
+	TOptional<FLobbySchemaId> LobbySchema;
+
+	/** New or changed attributes. */
 	TMap<FLobbyAttributeId, FLobbyVariant> MutatedAttributes;
 
-	// Attributes to be cleared.
+	/** Attributes to be cleared. */
 	TSet<FLobbyAttributeId> ClearedAttributes;
+
+	/** Members to be added or changed. */
+	TMap<FOnlineAccountIdHandle, TSharedRef<FClientLobbyMemberDataChanges>> MutatedMembers;
+
+	/** Members to be removed. */
+	TMap<FOnlineAccountIdHandle, ELobbyMemberLeaveReason> LeavingMembers;
 };
 
+/**
+* Common components required to handle lobby requests.
+* To handle lifetime issues, only FLobbiesEOS should contain a strong reference to prerequisites.
+* Any other component should only store a weak reference.
+*/
 struct FLobbyPrerequisitesEOS
 {
 	EOS_HLobby LobbyInterfaceHandle = {};
-	TSharedRef<FAuthEOS> AuthInterface;
+	TWeakPtr<FAuthEOS> AuthInterface;
 	TSharedRef<const FLobbySchemaRegistry> SchemaRegistry;
 	TSharedRef<const FLobbySchema> ServiceSchema;
 };
 
+/**
+* Lobby snapshot data.
+* All contained data has been translated from EOS types.
+* Attributes have had their schema transformations applied.
+*/
+struct FClientLobbySnapshot
+{
+	FOnlineAccountIdHandle OwnerAccountId;
+	FName LocalName;
+	FName SchemaName;
+	int32 MaxMembers;
+	ELobbyJoinPolicy JoinPolicy;
+	TMap<FLobbyAttributeId, FLobbyVariant> Attributes;
+	TSet<FOnlineAccountIdHandle> Members;
+};
+
+struct FClientLobbyMemberSnapshot : public FLobbyMember
+{
+	bool bIsLocalMember = false;
+};
+
+struct FApplyLobbyUpdateResult
+{
+	TArray<FOnlineAccountIdHandle> LeavingLocalMembers;
+};
+
+/** Lobby data as seen by the client. */
+struct FClientLobbyData final
+{
+public:
+	FClientLobbyData(FOnlineLobbyIdHandle LobbyId);
+
+	TSharedRef<const FLobby> GetPublicDataPtr() const { return PublicData; }
+	const FLobby& GetPublicData() const { return *PublicData; }
+
+	TSharedPtr<const FClientLobbyMemberSnapshot> GetMemberData(FOnlineAccountIdHandle MemberAccountId) const;
+
+	/**
+	 * Apply updated lobby data and generate changes.
+	 * The LeaveReason provides context for members who have left since the most recent snapshot.
+	 * Changing a lobby generates events for the local client.
+	 */
+	FApplyLobbyUpdateResult ApplyLobbyUpdateFromServiceSnapshot(
+		FClientLobbySnapshot&& LobbySnapshot,
+		TMap<FOnlineAccountIdHandle, TSharedRef<FClientLobbyMemberSnapshot>>&& LobbyMemberSnapshots,
+		TMap<FOnlineAccountIdHandle, ELobbyMemberLeaveReason>&& LeaveReasons = TMap<FOnlineAccountIdHandle, ELobbyMemberLeaveReason>(),
+		FLobbyEvents* LobbyEvents = nullptr);
+
+	/**
+	 * Apply updated lobby data and generate changes.
+	 * Changing a lobby generates events for the local client.
+	 */
+	FApplyLobbyUpdateResult ApplyLobbyUpdateFromLocalChanges(FClientLobbyDataChanges&& Changes, FLobbyEvents& LobbyEvents);
+
+private:
+
+	/** Apply changes to a set of attributes. Returns the set of attribute IDs which changed. */
+	TSet<FLobbyAttributeId> ApplyAttributeUpdateFromSnapshot(
+		TMap<FLobbyAttributeId, FLobbyVariant>&& AttributeSnapshot,
+		TMap<FLobbyAttributeId, FLobbyVariant>& ExistingAttributes);
+
+	/** Apply changes to a set of attributes. Returns the set of attribute IDs which changed. */
+	TSet<FLobbyAttributeId> ApplyAttributeUpdateFromChanges(
+		TMap<FLobbyAttributeId, FLobbyVariant>&& MutatedAttributes,
+		TSet<FLobbyAttributeId>&& ClearedAttributes,
+		TMap<FLobbyAttributeId, FLobbyVariant>& ExistingAttributes);
+
+	/**
+	 * The shared pointer given back to user code with lobby operation results and notifications.
+	 * Any changes to this data are immediately available to users.
+	 */
+	TSharedRef<FLobby> PublicData;
+
+	/** Mutable lobby member data storage. */
+	TMap<FOnlineAccountIdHandle, TSharedRef<FClientLobbyMemberSnapshot>> MemberDataStorage;
+
+	/**
+	 * Keep track of which members are local to the client.
+	 * When all local members have been removed all members will be removed.
+	 */
+	TSet<FOnlineAccountIdHandle> LocalMembers;
+};
+
+/**
+ * Lobby details are created for each user within the EOS lobby client. Certain operations such as
+ * joining a lobby or applying an update to a lobby / member attributes require using the correct
+ * lobby details handle.
+ */
 enum class ELobbyDetailsSource
 {
+	/**
+	 * The lobby has been joined and is considered active.
+	 * Valid for lobby / member updates.
+	 */
 	Active,
+
+	/**
+	 * The details originated from an invitation sent to the user.
+	 * Valid for joining.
+	 */
 	Invite,
+
+	/**
+	 * The details originated from an event generated by the EOS overlay.
+	 * Valid for joining.
+	 */
 	UiEvent,
+
+	/**
+	 * The details originated from a lobby search result.
+	 * Valid for joining.
+	 */
 	Search
 };
 
+/** Lobby details is created based on the passed in user and is required to join a lobby. */
 class FLobbyDetailsEOS final : public TSharedFromThis<FLobbyDetailsEOS>
 {
 public:
 	UE_NONCOPYABLE(FLobbyDetailsEOS);
 
-	static TDefaultErrorResultInternal<TSharedPtr<FLobbyDetailsEOS>> CreateFromLobbyId(const TSharedRef<FLobbyPrerequisitesEOS>& Prerequisites, EOS_LobbyId LobbyId, FOnlineAccountIdHandle MemberHandle);
-	static TDefaultErrorResultInternal<TSharedPtr<FLobbyDetailsEOS>> CreateFromInviteId(const TSharedRef<FLobbyPrerequisitesEOS>& Prerequisites, const char* InviteId);
-	static TDefaultErrorResultInternal<TSharedPtr<FLobbyDetailsEOS>> CreateFromUiEventId(const TSharedRef<FLobbyPrerequisitesEOS>& Prerequisites, EOS_UI_EventId UiEventId);
-	static TDefaultErrorResultInternal<TSharedPtr<FLobbyDetailsEOS>> CreateFromSearchResult(const TSharedRef<FLobbyPrerequisitesEOS>& Prerequisites, EOS_HLobbySearch SearchHandle, uint32_t ResultIndex);
+	static TDefaultErrorResultInternal<TSharedRef<FLobbyDetailsEOS>> CreateFromLobbyId(
+		const TSharedRef<FLobbyPrerequisitesEOS>& Prerequisites,
+		FOnlineAccountIdHandle LocalUserId,
+		EOS_LobbyId LobbyId);
+	static TDefaultErrorResultInternal<TSharedRef<FLobbyDetailsEOS>> CreateFromInviteId(
+		const TSharedRef<FLobbyPrerequisitesEOS>& Prerequisites,
+		FOnlineAccountIdHandle LocalUserId,
+		const char* InviteId);
+	static TDefaultErrorResultInternal<TSharedRef<FLobbyDetailsEOS>> CreateFromUiEventId(
+		const TSharedRef<FLobbyPrerequisitesEOS>& Prerequisites,
+		FOnlineAccountIdHandle LocalUserId,
+		EOS_UI_EventId UiEventId);
+	static TDefaultErrorResultInternal<TSharedRef<FLobbyDetailsEOS>> CreateFromSearchResult(
+		const TSharedRef<FLobbyPrerequisitesEOS>& Prerequisites,
+		FOnlineAccountIdHandle LocalUserId,
+		EOS_HLobbySearch SearchHandle,
+		uint32_t ResultIndex);
 
 	FLobbyDetailsEOS() = delete;
 	~FLobbyDetailsEOS();
 
-	bool IsValid() const { return LobbyDetailsHandle != InvalidLobbyDetailsHandle; }
 	EOS_HLobbyDetails GetEOSHandle() const { return LobbyDetailsHandle; }
 	const TSharedRef<FLobbyDetailsInfoEOS>& GetInfo() const { return LobbyDetailsInfo; }
 	ELobbyDetailsSource GetDetailsSource() const { return LobbyDetailsSource; }
-
-	TFuture<TDefaultErrorResultInternal<TSharedPtr<FLobby>>> GetLobbyData(FOnlineAccountIdHandle LocalUserId, FOnlineLobbyIdHandle LobbyIdHandle) const;
-	TDefaultErrorResultInternal<TSharedPtr<FLobbyMember>> GetLobbyMemberData(FOnlineAccountIdHandle MemberHandle) const;
-	TFuture<EOS_EResult> ApplyLobbyDataUpdates(FOnlineAccountIdHandle LocalUserId, FLobbyDataChanges Changes) const;
-	TFuture<EOS_EResult> ApplyLobbyMemberDataUpdates(FOnlineAccountIdHandle LocalUserId, FLobbyMemberDataChanges Changes) const;
+	FOnlineAccountIdHandle GetAssociatedUser() const { return AssociatedLocalUser; }
+
+	/**
+	 * Retrieve a lobby data snapshot from the EOS lobby details object.
+	 * The lobby schema will be used to translate attribute data before returning.
+	 */
+	TFuture<TDefaultErrorResultInternal<TSharedRef<FClientLobbySnapshot>>> GetLobbySnapshot() const;
+
+	/**
+	 * Retrieve lobby member data snapshot from the EOS lobby details object.
+	 * The lobby schema will be used to translate attribute data before returning.
+	 */
+	TDefaultErrorResultInternal<TSharedRef<FClientLobbyMemberSnapshot>> GetLobbyMemberSnapshot(FOnlineAccountIdHandle MemberHandle) const;
+
+	/**
+	 * Apply client side lobby changes to the lobby service.
+	 * The lobby schema will be used to translate any changed attributes before sending.
+	 */
+	TFuture<EOS_EResult> ApplyLobbyDataUpdateFromLocalChanges(FOnlineAccountIdHandle LocalUserId, const FClientLobbyDataChanges& Changes) const;
+
+	/**
+	 * Apply client side lobby member changes to the lobby service.
+	 * The lobby schema will be used to translate any changed attributes before sending.
+	 */
+	TFuture<EOS_EResult> ApplyLobbyMemberDataUpdateFromLocalChanges(FOnlineAccountIdHandle LocalUserId, const FClientLobbyMemberDataChanges& Changes) const;
 
 private:
 	template <typename, ESPMode>
 	friend class SharedPointerInternals::TIntrusiveReferenceController;
-	FLobbyDetailsEOS(const TSharedRef<FLobbyPrerequisitesEOS>& Prerequisites, const TSharedRef<FLobbyDetailsInfoEOS>& LobbyDetailsInfo, ELobbyDetailsSource LobbyDetailsSource, EOS_HLobbyDetails LobbyDetailsHandle);
-
-	TPair<FLobbyAttributeId, FLobbyVariant> TranslateLobbyAttribute(const EOS_Lobby_Attribute& LobbyAttribute) const;
-	TDefaultErrorResultInternal<TMap<FLobbyAttributeId, FLobbyVariant>> GetLobbyAttributes() const;
-	TDefaultErrorResultInternal<TMap<FLobbyAttributeId, FLobbyVariant>> GetLobbyMemberAttributes(EOS_ProductUserId TargetMemberProductUserId) const;
-	TSharedRef<FLobbyMember> CreateLobbyMember(FOnlineAccountIdHandle MemberHandle, TMap<FLobbyAttributeId, FLobbyVariant> Attributes) const;
+	FLobbyDetailsEOS(
+		const TSharedRef<FLobbyPrerequisitesEOS>& Prerequisites,
+		const TSharedRef<FLobbyDetailsInfoEOS>& LobbyDetailsInfo,
+		FOnlineAccountIdHandle LocalUserId,
+		ELobbyDetailsSource LobbyDetailsSource,
+		EOS_HLobbyDetails LobbyDetailsHandle);
 
 	static const EOS_HLobbyDetails InvalidLobbyDetailsHandle;
 	TSharedRef<FLobbyPrerequisitesEOS> Prerequisites;
 	TSharedRef<FLobbyDetailsInfoEOS> LobbyDetailsInfo;
-	EOS_HLobbyDetails LobbyDetailsHandle = InvalidLobbyDetailsHandle;
+	FOnlineAccountIdHandle AssociatedLocalUser;
 	ELobbyDetailsSource LobbyDetailsSource;
+	EOS_HLobbyDetails LobbyDetailsHandle = InvalidLobbyDetailsHandle;
 };
 
+/** Lobby details info allows direct access to some of the lobby properties. */
 class FLobbyDetailsInfoEOS final
 {
 public:
 	UE_NONCOPYABLE(FLobbyDetailsInfoEOS);
 	FLobbyDetailsInfoEOS() = default;
 
-	static TDefaultErrorResultInternal<TSharedPtr<FLobbyDetailsInfoEOS>> Create(EOS_HLobbyDetails LobbyDetailsHandle);
-
-	bool IsValid() const { return LobbyDetailsInfo.IsValid(); }
+	static TDefaultErrorResultInternal<TSharedRef<FLobbyDetailsInfoEOS>> Create(EOS_HLobbyDetails LobbyDetailsHandle);
 
-	EOS_LobbyId GetLobbyIdEOS() const { return LobbyDetailsInfo->LobbyId; }
+	const EOS_LobbyDetails_Info& Get() const { check(LobbyDetailsInfo); return *LobbyDetailsInfo; }
 
 private:
 	struct FEOSLobbyDetailsInfoDeleter
@@ -132,17 +312,16 @@ private:
 		}
 	};
 
-	TUniquePtr<EOS_LobbyDetails_Info, FEOSLobbyDetailsInfoDeleter> LobbyDetailsInfo;
-};
+	using FLobbyDetailsInfoPtr = TUniquePtr<EOS_LobbyDetails_Info, FEOSLobbyDetailsInfoDeleter>;
 
-class FLobbyNotificationPauseHandle
-{
-public:
-	UE_NONCOPYABLE(FLobbyNotificationPauseHandle);
-	FLobbyNotificationPauseHandle() = default;
-	virtual ~FLobbyNotificationPauseHandle() = default;
+	template <typename, ESPMode>
+	friend class SharedPointerInternals::TIntrusiveReferenceController;
+	FLobbyDetailsInfoEOS(FLobbyDetailsInfoPtr&& LobbyDetailsInfo);
+
+	TUniquePtr<EOS_LobbyDetails_Info, FEOSLobbyDetailsInfoDeleter> LobbyDetailsInfo;
 };
 
+/** Lobby data is the bookkeeping object for a lobby. It contains the client-side representation of a lobby. */
 class FLobbyDataEOS final : public TSharedFromThis<FLobbyDataEOS>
 {
 public:
@@ -150,46 +329,40 @@ public:
 
 	FLobbyDataEOS() = delete;
 
-	static TFuture<TDefaultErrorResultInternal<TSharedPtr<FLobbyDataEOS>>> Create(
-		FOnlineAccountIdHandle LocalUserId,
+	static TFuture<TDefaultErrorResultInternal<TSharedRef<FLobbyDataEOS>>> Create(
 		FOnlineLobbyIdHandle LobbyIdHandle,
 		const TSharedRef<FLobbyDetailsEOS>& LobbyDetails,
 		FUnregisterFn UnregisterFn = FUnregisterFn());
 
 	~FLobbyDataEOS();
 
-	FOnlineLobbyIdHandle GetLobbyIdHandle() const { return LobbyImpl->LobbyId; }
-	const TSharedRef<FLobby>& GetLobbyImpl() const { return LobbyImpl; }
-	EOS_LobbyId GetLobbyIdEOS() const { return LobbyDetailsInfo->GetLobbyIdEOS(); }
+	FOnlineLobbyIdHandle GetLobbyIdHandle() const { return ClientLobbyData->GetPublicData().LobbyId; }
+	const TSharedRef<FClientLobbyData>& GetClientLobbyData() const { return ClientLobbyData; }
+	EOS_LobbyId GetLobbyIdEOS() const { return LobbyDetailsInfo->Get().LobbyId; }
 	const FString& GetLobbyId() const { return LobbyId; }
-	TSet<FOnlineAccountIdHandle>& GetLocalMembers() { return LocalMembers; }
-
-	// Pause lobby notifications for the target local user.
-	TSharedRef<FLobbyNotificationPauseHandle> PauseLobbyNotifications(FOnlineAccountIdHandle LocalUserId);
 
-	// Returns false if there is an existing notification pause for the target local user.
-	bool CanNotify(FOnlineAccountIdHandle LocalUserId);
-
-	void SetUserLobbyDetails(FOnlineAccountIdHandle LocalUserId, const TSharedPtr<FLobbyDetailsEOS>& LobbyDetails);
+	void AddUserLobbyDetails(FOnlineAccountIdHandle LocalUserId, const TSharedPtr<FLobbyDetailsEOS>& LobbyDetails);
 	TSharedPtr<FLobbyDetailsEOS> GetUserLobbyDetails(FOnlineAccountIdHandle LocalUserId) const;
 
+	/**
+	 * Active lobby details are needed to process lobby notifications. Search for and return active
+	 * lobby details if available.
+	 */
+	TSharedPtr<FLobbyDetailsEOS> GetActiveLobbyDetails() const;
+
 private:
 	template <typename, ESPMode>
 	friend class SharedPointerInternals::TIntrusiveReferenceController;
 
 	FLobbyDataEOS(
-		FOnlineLobbyIdHandle LobbyIdHandle,
-		const TSharedRef<FLobby>& LobbyImpl,
+		const TSharedRef<FClientLobbyData>& ClientLobbyData,
 		const TSharedRef<FLobbyDetailsInfoEOS>& LobbyDetailsInfo,
 		FUnregisterFn UnregisterFn);
 
-	FOnlineLobbyIdHandle LobbyIdHandle;
-	TSharedRef<FLobby> LobbyImpl;
+	TSharedRef<FClientLobbyData> ClientLobbyData;
 	TSharedRef<FLobbyDetailsInfoEOS> LobbyDetailsInfo;
 	FUnregisterFn UnregisterFn;
 	FString LobbyId;
-	TSet<FOnlineAccountIdHandle> LocalMembers;
-	TMap<FOnlineAccountIdHandle, TSharedRef<FLobbyNotificationPauseHandle>> NotificationPauses;
 	TMap<FOnlineAccountIdHandle, TSharedPtr<FLobbyDetailsEOS>> UserLobbyDetails;
 };
 
@@ -201,9 +374,7 @@ public:
 
 	TSharedPtr<FLobbyDataEOS> Find(EOS_LobbyId EOSLobbyId) const;
 	TSharedPtr<FLobbyDataEOS> Find(FOnlineLobbyIdHandle LobbyIdHandle) const;
-	TFuture<TDefaultErrorResultInternal<TSharedPtr<FLobbyDataEOS>>> FindOrCreateFromLobbyId(FOnlineAccountIdHandle LocalUserId, EOS_LobbyId LobbyId);
-	TFuture<TDefaultErrorResultInternal<TSharedPtr<FLobbyDataEOS>>> FindOrCreateFromUiEventId(FOnlineAccountIdHandle LocalUserId, EOS_UI_EventId UiEventId);
-	TFuture<TDefaultErrorResultInternal<TSharedPtr<FLobbyDataEOS>>> FindOrCreateFromLobbyDetails(FOnlineAccountIdHandle LocalUserId, const TSharedRef<FLobbyDetailsEOS>& LobbyDetails);
+	TFuture<TDefaultErrorResultInternal<TSharedRef<FLobbyDataEOS>>> FindOrCreateFromLobbyDetails(FOnlineAccountIdHandle LocalUserId, const TSharedRef<FLobbyDetailsEOS>& LobbyDetails);
 
 private:
 	void Register(const TSharedRef<FLobbyDataEOS>& LobbyIdHandleData);
@@ -232,12 +403,15 @@ private:
 	char Data[MaxLobbyInviteIdSize] = {};
 };
 
+/**
+* Lobby invite data will keep the lobby object valid until the invitation has been accepted or rejected.
+*/
 class FLobbyInviteDataEOS final
 {
 public:
 	FLobbyInviteDataEOS() = delete;
 
-	static TFuture<TDefaultErrorResultInternal<TSharedPtr<FLobbyInviteDataEOS>>> CreateFromInviteId(
+	static TFuture<TDefaultErrorResultInternal<TSharedRef<FLobbyInviteDataEOS>>> CreateFromInviteId(
 		const TSharedRef<FLobbyPrerequisitesEOS>& Prerequisites,
 		const TSharedRef<FLobbyDataRegistryEOS>& LobbyDataRegistry,
 		FOnlineAccountIdHandle LocalUserId,
@@ -272,20 +446,22 @@ private:
 	FString InviteId;
 };
 
-// The lobby search object is meant to hold the lifetimes of lobby search results so that the
-// lobby details will be discoverable during a join operation.
+/**
+* The lobby search object is meant to hold the lifetimes of lobby search results so that the
+* lobby details will be discoverable during a join operation.
+*/
 class FLobbySearchEOS final
 {
 public:
 	FLobbySearchEOS() = delete;
 
-	static TFuture<TDefaultErrorResultInternal<TSharedPtr<FLobbySearchEOS>>> Create(
+	static TFuture<TDefaultErrorResultInternal<TSharedRef<FLobbySearchEOS>>> Create(
 		const TSharedRef<FLobbyPrerequisitesEOS>& Prerequisites,
 		const TSharedRef<FLobbyDataRegistryEOS>& LobbyRegistry,
 		const FLobbySearchParameters& Params);
 
 	TArray<TSharedRef<const FLobby>> GetLobbyResults() const;
-	const TArray<TSharedPtr<FLobbyDataEOS>>& GetLobbyData();
+	const TArray<TSharedRef<FLobbyDataEOS>>& GetLobbyData();
 
 private:
 	class FSearchHandle
@@ -310,10 +486,10 @@ private:
 
 	template <typename, ESPMode>
 	friend class SharedPointerInternals::TIntrusiveReferenceController;
-	FLobbySearchEOS(const TSharedRef<FSearchHandle>& SearchHandle, TArray<TSharedPtr<FLobbyDataEOS>>&& Lobbies);
+	FLobbySearchEOS(const TSharedRef<FSearchHandle>& SearchHandle, TArray<TSharedRef<FLobbyDataEOS>>&& Lobbies);
 
 	TSharedRef<FSearchHandle> SearchHandle;
-	TArray<TSharedPtr<FLobbyDataEOS>> Lobbies;
+	TArray<TSharedRef<FLobbyDataEOS>> Lobbies;
 };
 
 /* UE::Online */ }
diff --git a/Engine/Plugins/Online/Experimental/OnlineServicesEOS/Source/Public/LobbiesEOS.h b/Engine/Plugins/Online/Experimental/OnlineServicesEOS/Source/Public/LobbiesEOS.h
--- a/Engine/Plugins/Online/Experimental/OnlineServicesEOS/Source/Public/LobbiesEOS.h
+++ b/Engine/Plugins/Online/Experimental/OnlineServicesEOS/Source/Public/LobbiesEOS.h
@@ -13,6 +13,7 @@
 
 namespace UE::Online {
 
+class FAuthEOS;
 struct FLobbyPrerequisitesEOS;
 class FLobbyDataEOS;
 class FLobbyDataRegistryEOS;
@@ -20,8 +21,8 @@ class FLobbyDetailsEOS;
 class FLobbyInviteDataEOS;
 class FOnlineServicesEOS;
 class FLobbySearchEOS;
-struct FLobbyDataChanges;
-struct FLobbyMemberDataChanges;
+struct FClientLobbyDataChanges;
+struct FClientLobbyMemberDataChanges;
 	
 class FLobbiesEOS : public FLobbiesCommon
 {
@@ -32,7 +33,7 @@ public:
 	virtual void PreShutdown() override;
 
 	virtual TOnlineAsyncOpHandle<FCreateLobby> CreateLobby(FCreateLobby::Params&& Params) override;
-	virtual TOnlineAsyncOpHandle<FFindLobby> FindLobby(FFindLobby::Params&& Params) override;
+	virtual TOnlineAsyncOpHandle<FFindLobbies> FindLobbies(FFindLobbies::Params&& Params) override;
 	virtual TOnlineAsyncOpHandle<FJoinLobby> JoinLobby(FJoinLobby::Params&& Params) override;
 	virtual TOnlineAsyncOpHandle<FLeaveLobby> LeaveLobby(FLeaveLobby::Params&& Params) override;
 	virtual TOnlineAsyncOpHandle<FInviteLobbyMember> InviteLobbyMember(FInviteLobbyMember::Params&& Params) override;
@@ -68,7 +69,7 @@ protected:
 	{
 		static constexpr TCHAR Name[] = TEXT("JoinLobbyImpl");
 
-		struct Params : public FJoinLobby::Params
+		struct Params
 		{
 			// The lobby handle data.
 			TSharedPtr<FLobbyDataEOS> LobbyData;
@@ -113,6 +114,7 @@ protected:
 	};
 
 	TOnlineAsyncOpHandle<FJoinLobbyMemberImpl> JoinLobbyMemberImplOp(FJoinLobbyMemberImpl::Params&& Params);
+	TFuture<TDefaultErrorResult<FJoinLobbyMemberImpl>> JoinLobbyMemberImpl(FJoinLobbyMemberImpl::Params&& Params);
 
 	struct FLeaveLobbyImpl
 	{
@@ -256,9 +258,12 @@ protected:
 			FOnlineAccountIdHandle LocalUserId;
 
 			// The changes to apply to the lobby data.
-			TSharedPtr<FLobbyDataChanges> Changes;
+			TSharedPtr<FClientLobbyDataChanges> Changes;
 		};
 
+		// Todo: += operator.
+		// Mergeable op must be queued by lobby id.
+
 		struct Result
 		{
 		};
@@ -279,9 +284,12 @@ protected:
 			FOnlineAccountIdHandle LocalUserId;
 
 			// The changes to apply to the lobby member data.
-			TSharedPtr<FLobbyMemberDataChanges> Changes;
+			TSharedPtr<FClientLobbyMemberDataChanges> Changes;
 		};
 
+		// Todo: += operator.
+		// Mergeable op must be queued on union of lobby id + lobby member account id.
+
 		struct Result
 		{
 		};
@@ -289,6 +297,32 @@ protected:
 
 	TFuture<TDefaultErrorResult<FModifyLobbyMemberDataImpl>> ModifyLobbyMemberDataImpl(FModifyLobbyMemberDataImpl::Params&& Params);
 
+	struct FProcessLobbyNotificationImpl
+	{
+		static constexpr TCHAR Name[] = TEXT("ProcessLobbyNotificationImpl");
+
+		struct Params
+		{
+			// The lobby handle data.
+			TSharedPtr<FLobbyDataEOS> LobbyData;
+
+			// Joining / mutated members.
+			TSet<EOS_ProductUserId> MutatedMembers;
+
+			// Leaving members.
+			TMap<EOS_ProductUserId, ELobbyMemberLeaveReason> LeavingMembers;
+		};
+
+		// Todo: += operator.
+		// Mergeable op must be queued by lobby id.
+
+		struct Result
+		{
+		};
+	};
+
+	TOnlineAsyncOpHandle<FProcessLobbyNotificationImpl> ProcessLobbyNotificationImplOp(FProcessLobbyNotificationImpl::Params&& Params);
+
 	FString LobbyBucketId;
 
 	EOSEventRegistrationPtr OnLobbyUpdatedEOSEventRegistration;
