Change directory structure for OnlineServices, OnlineservicesEOS, and OnlineServicesNULL to be consistent.

diff --git a/Engine/Plugins/Online/Experimental/OnlineServicesEOS/Source/Private/Online/ExternalUIEOS.cpp b/Engine/Plugins/Online/Experimental/OnlineServicesEOS/Source/Private/Online/ExternalUIEOS.cpp
new file mode 100644
--- /dev/null
+++ b/Engine/Plugins/Online/Experimental/OnlineServicesEOS/Source/Private/Online/ExternalUIEOS.cpp
@@ -0,0 +1,121 @@
+// Copyright Epic Games, Inc. All Rights Reserved.
+
+#include "Online/ExternalUIEOS.h"
+
+#include "Online/AuthEOS.h"
+#include "Online/OnlineIdEOS.h"
+#include "Online/OnlineServicesEOS.h"
+#include "Online/OnlineServicesEOSTypes.h"
+#include "Online/OnlineErrorDefinitions.h"
+
+#include "eos_ui.h"
+#include "eos_types.h"
+#include "eos_init.h"
+#include "eos_sdk.h"
+#include "eos_logging.h"
+
+namespace UE::Online {
+
+FExternalUIEOS::FExternalUIEOS(FOnlineServicesEOS& InServices)
+	: FExternalUICommon(InServices)
+{
+}
+
+void FExternalUIEOS::Initialize()
+{
+	Super::Initialize();
+
+	UIHandle = EOS_Platform_GetUIInterface(static_cast<FOnlineServicesEOS&>(GetServices()).GetEOSPlatformHandle());
+	check(UIHandle != nullptr);
+}
+
+void FExternalUIEOS::LoadConfig()
+{
+	Super::LoadConfig();
+	LoadConfig(Config);
+}
+
+void FExternalUIEOS::PreShutdown()
+{
+	Super::PreShutdown();
+}
+
+TOnlineAsyncOpHandle<FExternalUIShowLoginUI> FExternalUIEOS::ShowLoginUI(FExternalUIShowLoginUI::Params&& Params)
+{
+	TOnlineAsyncOpRef<FExternalUIShowLoginUI> Op = GetOp<FExternalUIShowLoginUI>(MoveTemp(Params));
+	if (Config.bShowLoginUIEnabled)
+	{
+		// Use account portal auth login
+		Op->Then([this](TOnlineAsyncOp<FExternalUIShowLoginUI>& InAsyncOp)
+		{
+			FAuthLogin::Params AuthLoginParams;
+			AuthLoginParams.PlatformUserId = InAsyncOp.GetParams().PlatformUserId;
+			AuthLoginParams.CredentialsType = TEXT("accountportal");
+			AuthLoginParams.Scopes = InAsyncOp.GetParams().Scopes;
+
+			TPromise<void>* TempPromise = new TPromise<void>;
+			GetServices().GetAuthInterface()->Login(MoveTemp(AuthLoginParams))
+				.OnComplete([Op = InAsyncOp.AsShared(), TempPromise](const TOnlineResult<FAuthLogin>& AuthLoginResult) mutable
+				{
+					if (AuthLoginResult.IsOk())
+					{
+						FExternalUIShowLoginUI::Result ExtUIShowLoginResult = { AuthLoginResult.GetOkValue().AccountInfo };
+						Op->SetResult(CopyTemp(ExtUIShowLoginResult));
+					}
+					else
+					{
+						Op->SetError(CopyTemp(AuthLoginResult.GetErrorValue()));
+					}
+					TempPromise->SetValue();
+					delete TempPromise;
+				});
+			return TempPromise->GetFuture();
+		}).Enqueue(GetSerialQueue());
+	}
+	else
+	{
+		Op->SetError(Errors::NotImplemented());
+	}
+	return Op->GetHandle();
+}
+
+TOnlineAsyncOpHandle<FExternalUIShowFriendsUI> FExternalUIEOS::ShowFriendsUI(FExternalUIShowFriendsUI::Params&& Params)
+{
+	TOnlineAsyncOpRef<FExternalUIShowFriendsUI> Op = GetOp<FExternalUIShowFriendsUI>(MoveTemp(Params));
+
+	EOS_EpicAccountId LocalUserEasId = GetEpicAccountId(Params.LocalUserId);
+	if (!EOS_EpicAccountId_IsValid(LocalUserEasId))
+	{
+		UE_LOG(LogTemp, Warning, TEXT("[FExternalUIEOS::ShowFriendsUI] LocalUserId=[%s] EpicAccountId not found"), *ToLogString(Params.LocalUserId));
+		Op->SetError(Errors::Unknown()); // TODO
+		return Op->GetHandle();
+	}
+
+	EOS_UI_ShowFriendsOptions ShowFriendsOptions = {};
+	ShowFriendsOptions.ApiVersion = EOS_UI_SHOWFRIENDS_API_LATEST;
+	ShowFriendsOptions.LocalUserId = LocalUserEasId;
+
+	Op->Then([this, ShowFriendsOptions](TOnlineAsyncOp<FExternalUIShowFriendsUI>& InAsyncOp) mutable
+	{
+		return EOS_Async<EOS_UI_ShowFriendsCallbackInfo>(EOS_UI_ShowFriends, UIHandle, ShowFriendsOptions);
+	})
+	.Then([this](TOnlineAsyncOp<FExternalUIShowFriendsUI>& InAsyncOp, const EOS_UI_ShowFriendsCallbackInfo* Data)
+	{
+		UE_LOG(LogTemp, Verbose, TEXT("[FExternalUIEOS::ShowFriendsUI] EOS_UI_ShowFriends Result [%s] for User [%s]."), *LexToString(Data->ResultCode), *LexToString(Data->LocalUserId));
+
+		if (Data->ResultCode == EOS_EResult::EOS_Success)
+		{
+			FExternalUIShowFriendsUI::Result Result = {};
+			InAsyncOp.SetResult(MoveTemp(Result));
+		}
+		else
+		{
+			InAsyncOp.SetError(Errors::Unknown());
+		}
+	})
+	.Enqueue(GetSerialQueue());
+
+	return Op->GetHandle();
+}
+
+/* UE::Online */ }
diff --git a/Engine/Plugins/Online/Experimental/OnlineServicesEOS/Source/Private/Online/FriendsEOS.cpp b/Engine/Plugins/Online/Experimental/OnlineServicesEOS/Source/Private/Online/FriendsEOS.cpp
new file mode 100644
--- /dev/null
+++ b/Engine/Plugins/Online/Experimental/OnlineServicesEOS/Source/Private/Online/FriendsEOS.cpp
@@ -0,0 +1,302 @@
+// Copyright Epic Games, Inc. All Rights Reserved.
+
+#include "Online/FriendsEOS.h"
+#include "Online/OnlineAsyncOp.h"
+#include "Online/OnlineServicesEOS.h"
+#include "Online/OnlineServicesEOSTypes.h"
+#include "Online/OnlineIdEOS.h"
+#include "Online/AuthEOS.h"
+
+#include "eos_friends.h"
+
+namespace UE::Online {
+
+inline TOptional<EFriendInviteStatus> EOSFriendStatusToInviteStatus(EOS_EFriendsStatus EOSFriendStatus)
+{
+	TOptional<EFriendInviteStatus> InviteStatus;
+	switch (EOSFriendStatus)
+	{
+	case EOS_EFriendsStatus::EOS_FS_InviteSent:
+		InviteStatus.Emplace(EFriendInviteStatus::PendingOutbound);
+		break;
+	case EOS_EFriendsStatus::EOS_FS_InviteReceived:
+		InviteStatus.Emplace(EFriendInviteStatus::PendingInbound);
+		break;
+	}
+	return InviteStatus;
+}
+
+FFriendsEOS::FFriendsEOS(FOnlineServicesEOS& InServices)
+	: FFriendsCommon(InServices)
+{
+}
+
+void FFriendsEOS::Initialize()
+{
+	FFriendsCommon::Initialize();
+
+	FriendsHandle = EOS_Platform_GetFriendsInterface(static_cast<FOnlineServicesEOS&>(GetServices()).GetEOSPlatformHandle());
+	check(FriendsHandle != nullptr);
+
+	// Register for friend updates
+	EOS_Friends_AddNotifyFriendsUpdateOptions Options = { };
+	Options.ApiVersion = EOS_FRIENDS_ADDNOTIFYFRIENDSUPDATE_API_LATEST;
+	NotifyFriendsUpdateNotificationId = EOS_Friends_AddNotifyFriendsUpdate(FriendsHandle, &Options, this, [](const EOS_Friends_OnFriendsUpdateInfo* Data)
+	{
+		FFriendsEOS* This = reinterpret_cast<FFriendsEOS*>(Data->ClientData);
+
+		const FOnlineAccountIdHandle LocalUserId = FindAccountIdChecked(Data->LocalUserId);
+		This->Services.Get<FAuthEOS>()->ResolveAccountId(LocalUserId, Data->TargetUserId)
+		.Next([This, LocalUserId, PreviousStatus = Data->PreviousStatus, CurrentStatus = Data->CurrentStatus](const FOnlineAccountIdHandle& FriendUserId)
+		{
+			UE_LOG(LogTemp, Warning, TEXT("OnEOSFriendsUpdate: LocalUserId=[%s] FriendUserId=[%s]"), *ToLogString(LocalUserId), *ToLogString(FriendUserId));
+			This->OnEOSFriendsUpdate(LocalUserId, FriendUserId, PreviousStatus, CurrentStatus);
+		});
+	});
+}
+
+void FFriendsEOS::PreShutdown()
+{
+	EOS_Friends_RemoveNotifyFriendsUpdate(FriendsHandle, NotifyFriendsUpdateNotificationId);
+}
+
+TOnlineAsyncOpHandle<FQueryFriends> FFriendsEOS::QueryFriends(FQueryFriends::Params&& InParams)
+{
+	TOnlineAsyncOpRef<FQueryFriends> Op = GetJoinableOp<FQueryFriends>(MoveTemp(InParams));
+	if (!Op->IsReady())
+	{
+		// Initialize
+		const FQueryFriends::Params& Params = Op->GetParams();
+		if (!Services.Get<FAuthEOS>()->IsLoggedIn(Params.LocalUserId))
+		{
+			// TODO: Error codes
+			Op->SetError(Errors::Unknown());
+			return Op->GetHandle();
+		}
+
+		EOS_Friends_QueryFriendsOptions QueryFriendsOptions = { };
+		QueryFriendsOptions.ApiVersion = EOS_FRIENDS_QUERYFRIENDS_API_LATEST;
+		EOS_EpicAccountId LocalUserEasId = GetEpicAccountIdChecked(Params.LocalUserId);
+		QueryFriendsOptions.LocalUserId = LocalUserEasId;
+
+		Op->Then([this, QueryFriendsOptions](TOnlineAsyncOp<FQueryFriends>& InAsyncOp) mutable
+		{
+			return EOS_Async<EOS_Friends_QueryFriendsCallbackInfo>(EOS_Friends_QueryFriends, FriendsHandle, QueryFriendsOptions);
+		})
+		.Then([this](TOnlineAsyncOp<FQueryFriends>& InAsyncOp, const EOS_Friends_QueryFriendsCallbackInfo* Data) mutable
+		{
+			UE_LOG(LogTemp, Warning, TEXT("QueryFriendsResult: [%s]"), *LexToString(Data->ResultCode));
+
+			if (Data->ResultCode != EOS_EResult::EOS_Success)
+			{
+				InAsyncOp.SetError(Errors::Unknown()); // TODO: Error codes
+				return MakeFulfilledPromise<TArray<EOS_EpicAccountId>>().GetFuture();
+			}
+
+			EOS_Friends_GetFriendsCountOptions GetFriendsCountOptions = {
+				EOS_FRIENDS_GETFRIENDSCOUNT_API_LATEST,
+				Data->LocalUserId
+			};
+
+			int32_t NumFriends = EOS_Friends_GetFriendsCount(FriendsHandle, &GetFriendsCountOptions);
+			TArray<EOS_EpicAccountId> CurrentFriends;
+			CurrentFriends.Reserve(NumFriends);
+			for (int32_t FriendIndex = 0; FriendIndex < NumFriends; ++FriendIndex)
+			{
+				EOS_Friends_GetFriendAtIndexOptions GetFriendAtIndexOptions = {
+					EOS_FRIENDS_GETFRIENDATINDEX_API_LATEST,
+					Data->LocalUserId,
+					FriendIndex
+				};
+
+				const EOS_EpicAccountId FriendEasId = EOS_Friends_GetFriendAtIndex(FriendsHandle, &GetFriendAtIndexOptions);
+				check(EOS_EpicAccountId_IsValid(FriendEasId));
+				CurrentFriends.Emplace(FriendEasId);
+			}
+			return MakeFulfilledPromise<TArray<EOS_EpicAccountId>>(CurrentFriends).GetFuture();
+		}).Then(Services.Get<FAuthEOS>()->ResolveEpicIdsFn())
+		.Then([this](TOnlineAsyncOp<FQueryFriends>& InAsyncOp, const TArray<FOnlineAccountIdHandle>& CurrentFriendIds) mutable
+		{
+			const FOnlineAccountIdHandle LocalUserId = InAsyncOp.GetParams().LocalUserId;
+			const EOS_EpicAccountId LocalUserEasId = GetEpicAccountIdChecked(LocalUserId);
+
+			TMap<FOnlineAccountIdHandle, TSharedRef<FFriend>>& FriendsList = FriendsLists.FindOrAdd(LocalUserId);
+			TArray<FOnlineAccountIdHandle> NewFriendIds;
+			TArray<FOnlineAccountIdHandle> PreviousFriendIds;
+			TArray<FOnlineAccountIdHandle> UpdatedFriendIds; // TODO:  More granular updates like what changed?
+			FriendsList.GenerateKeyArray(PreviousFriendIds);
+				
+			for (const FOnlineAccountIdHandle& FriendId : CurrentFriendIds)
+			{
+				const EOS_EpicAccountId FriendEasId = GetEpicAccountIdChecked(FriendId);
+
+				EOS_Friends_GetStatusOptions GetStatusOptions = {
+					EOS_FRIENDS_GETSTATUS_API_LATEST,
+					LocalUserEasId,
+					FriendEasId
+				};
+				const EOS_EFriendsStatus EOSFriendStatus = EOS_Friends_GetStatus(FriendsHandle, &GetStatusOptions);
+				TOptional<EFriendInviteStatus> InviteStatus = EOSFriendStatusToInviteStatus(EOSFriendStatus);
+
+				PreviousFriendIds.Remove(FriendId);
+				if (TSharedRef<FFriend>* FriendPtr = FriendsList.Find(FriendId))
+				{
+					// Any updates?
+					if (InviteStatus != (*FriendPtr)->InviteStatus)
+					{
+						UE_LOG(LogTemp, Warning, TEXT("QueryFriendsComplete: %s Invite status changed"), *LexToString(FriendEasId));
+						(*FriendPtr)->InviteStatus = InviteStatus;
+						UpdatedFriendIds.Emplace(FriendId);
+					}
+				}
+				else
+				{
+					// Add new friend
+					UE_LOG(LogTemp, Warning, TEXT("QueryFriendsComplete: Adding %s"), *LexToString(FriendEasId));
+					TSharedRef<FFriend> Friend = MakeShared<FFriend>();
+					Friend->UserId = FriendId;
+					Friend->InviteStatus = InviteStatus;
+					FriendsList.Emplace(FriendId, MoveTemp(Friend));
+					NewFriendIds.Emplace(FriendId);
+				}
+			}
+
+			// TODO:  Delegates
+
+			for (const FOnlineAccountIdHandle& AccountId : PreviousFriendIds)
+			{
+				UE_LOG(LogTemp, Warning, TEXT("QueryFriendsComplete: %s is no longer a friend, removing"), *ToLogString(AccountId));
+				FriendsList.Remove(AccountId);
+			}
+
+			InAsyncOp.SetResult(FQueryFriends::Result());
+					
+			const bool bAnyChanges = NewFriendIds.Num() > 0
+				|| PreviousFriendIds.Num() > 0
+				|| UpdatedFriendIds.Num() > 0;
+			if (bAnyChanges)
+			{
+				FFriendsListUpdated FriendsListUpdatedParams = {};
+				FriendsListUpdatedParams.LocalUserId = InAsyncOp.GetParams().LocalUserId;
+				FriendsListUpdatedParams.AddedFriends = NewFriendIds;
+				FriendsListUpdatedParams.RemovedFriends = PreviousFriendIds;
+				FriendsListUpdatedParams.UpdatedFriends = UpdatedFriendIds;
+				OnFriendsListUpdatedEvent.Broadcast(FriendsListUpdatedParams);
+			}
+		})
+		.Enqueue(GetSerialQueue());
+	}
+	return Op->GetHandle();
+}
+
+TOnlineResult<FGetFriends> FFriendsEOS::GetFriends(FGetFriends::Params&& Params)
+{
+	if (TMap<FOnlineAccountIdHandle, TSharedRef<FFriend>>* FriendsList = FriendsLists.Find(Params.LocalUserId))
+	{
+		FGetFriends::Result Result;
+		FriendsList->GenerateValueArray(Result.Friends);
+		return TOnlineResult<FGetFriends>(MoveTemp(Result));
+	}
+	return TOnlineResult<FGetFriends>(Errors::Unknown()); // TODO: error codes
+}
+
+TOnlineAsyncOpHandle<FAddFriend> FFriendsEOS::AddFriend(FAddFriend::Params&& InParams)
+{
+	TOnlineAsyncOpRef<FAddFriend> Op = GetJoinableOp<FAddFriend>(MoveTemp(InParams));
+	if (!Op->IsReady())
+	{
+#if 0 // At present this causes issues due to the EOS function executing the callback immediately, before the future can be bound to
+		// Initialize
+		const FAddFriend::Params& Params = Op->GetParams();
+		if (!Services.Get<FAuthEOS>()->IsLoggedIn(Params.LocalUserId))
+		{
+			// TODO: Error codes
+			Op->SetError(Errors::UnknownError());
+			return Op->GetHandle();
+		}
+
+		// Target user valid?
+		EOS_EpicAccountId TargetId = GetEpicAccountId(Params.FriendId);
+		if (!EOS_EpicAccountId_IsValid(TargetId))
+		{
+			// TODO: Error codes
+			Op->SetError(Errors::UnknownError());
+			return Op->GetHandle();
+		}
+
+		Op->Then([this](TOnlineAsyncOp<FAddFriend>& InAsyncOp) mutable
+		{
+			const FAddFriend::Params& Params = InAsyncOp.GetParams();
+			EOS_Friends_SendInviteOptions SendInviteOptions = { };
+			SendInviteOptions.ApiVersion = EOS_FRIENDS_SENDINVITE_API_LATEST;
+			SendInviteOptions.LocalUserId = EOSAccountIdFromOnlineServiceAccountId(Params.LocalUserId).GetValue();
+			SendInviteOptions.TargetUserId = EOSAccountIdFromOnlineServiceAccountId(Params.FriendId).GetValue();
+			return EOS_Async<EOS_Friends_SendInviteCallbackInfo>(EOS_Friends_SendInvite, FriendsHandle, SendInviteOptions);
+		})
+		.Then([this](TOnlineAsyncOp<FAddFriend>& InAsyncOp, const EOS_Friends_SendInviteCallbackInfo* Data) mutable
+		{
+			UE_LOG(LogTemp, Warning, TEXT("SendInviteResult: [%s]"), *LexToString(Data->ResultCode));
+			// TODO:  Handle response
+		}).Enqueue();
+#else
+		Op->SetError(Errors::Unknown());
+#endif
+	}
+	return Op->GetHandle();
+}
+
+void FFriendsEOS::OnEOSFriendsUpdate(FOnlineAccountIdHandle LocalUserId, FOnlineAccountIdHandle FriendUserId, EOS_EFriendsStatus PreviousStatus, EOS_EFriendsStatus CurrentStatus)
+{
+	// TODO
+	UE_LOG(LogTemp, Warning, TEXT("OnEOSFriendsUpdate: LocalUserId=[%s] FriendUserId=[%s] PreviousStatus=[%s] CurrentStatus=[%s]"), *ToLogString(LocalUserId), *ToLogString(FriendUserId), *LexToString(PreviousStatus), *LexToString(CurrentStatus));
+
+	TMap<FOnlineAccountIdHandle, TSharedRef<FFriend>>& FriendsList = FriendsLists.FindOrAdd(LocalUserId);
+	bool bAnyChanges = false;
+	TArray<FOnlineAccountIdHandle, TInlineAllocator<1>> AddedFriends;
+	TArray<FOnlineAccountIdHandle, TInlineAllocator<1>> RemovedFriends;
+	TArray<FOnlineAccountIdHandle, TInlineAllocator<1>> UpdatedFriends;
+	if (CurrentStatus == EOS_EFriendsStatus::EOS_FS_NotFriends)
+	{
+		bAnyChanges = FriendsList.Remove(FriendUserId) > 0;
+		if (bAnyChanges)
+		{
+			RemovedFriends.Emplace(FriendUserId);
+		}
+	}
+	else
+	{
+		TOptional<EFriendInviteStatus> InviteStatus = EOSFriendStatusToInviteStatus(CurrentStatus);
+		if (TSharedRef<FFriend>* ExistingFriend = FriendsList.Find(FriendUserId))
+		{
+			if ((*ExistingFriend)->InviteStatus != InviteStatus)
+			{
+				(*ExistingFriend)->InviteStatus = InviteStatus;
+				UpdatedFriends.Emplace(FriendUserId);
+				bAnyChanges = true;
+			}
+		}
+		else
+		{
+			// Create new entry
+			TSharedRef<FFriend> Friend = MakeShared<FFriend>();
+			Friend->UserId = FriendUserId;
+			Friend->InviteStatus = InviteStatus;
+			FriendsList.Emplace(FriendUserId, MoveTemp(Friend));
+
+			AddedFriends.Emplace(FriendUserId);
+			bAnyChanges = true;
+		}
+	}
+
+	if (bAnyChanges)
+	{
+		FFriendsListUpdated FriendsListUpdatedParams = {};
+		FriendsListUpdatedParams.LocalUserId = LocalUserId;
+		FriendsListUpdatedParams.AddedFriends = AddedFriends;
+		FriendsListUpdatedParams.RemovedFriends = RemovedFriends;
+		FriendsListUpdatedParams.UpdatedFriends = UpdatedFriends;
+		OnFriendsListUpdatedEvent.Broadcast(FriendsListUpdatedParams);
+	}
+}
+
+/* UE::Online */ }
diff --git a/Engine/Plugins/Online/Experimental/OnlineServicesEOS/Source/Private/Online/PresenceEOS.cpp b/Engine/Plugins/Online/Experimental/OnlineServicesEOS/Source/Private/Online/PresenceEOS.cpp
new file mode 100644
--- /dev/null
+++ b/Engine/Plugins/Online/Experimental/OnlineServicesEOS/Source/Private/Online/PresenceEOS.cpp
@@ -0,0 +1,463 @@
+// Copyright Epic Games, Inc. All Rights Reserved.
+
+#include "Online/PresenceEOS.h"
+#include "Online/OnlineAsyncOp.h"
+#include "Online/OnlineIdEOS.h"
+#include "Online/OnlineServicesEOS.h"
+#include "Online/OnlineServicesEOSTypes.h"
+#include "Online/AuthEOS.h"
+
+#include "eos_presence.h"
+
+namespace UE::Online {
+
+static inline EPresenceState ToEPresenceState(EOS_Presence_EStatus InStatus)
+{
+	switch (InStatus)
+	{
+		case EOS_Presence_EStatus::EOS_PS_Online:
+		{
+			return EPresenceState::Online;
+		}
+		case EOS_Presence_EStatus::EOS_PS_Away:
+		{
+			return EPresenceState::Away;
+		}
+		case EOS_Presence_EStatus::EOS_PS_ExtendedAway:
+		{
+			return EPresenceState::ExtendedAway;
+		}
+		case EOS_Presence_EStatus::EOS_PS_DoNotDisturb:
+		{
+			return EPresenceState::DoNotDisturb;
+		}
+	}
+	return EPresenceState::Offline;
+}
+
+static inline EOS_Presence_EStatus ToEOS_Presence_EStatus(EPresenceState InState)
+{
+	switch (InState)
+	{
+	case EPresenceState::Online:
+	{
+		return EOS_Presence_EStatus::EOS_PS_Online;
+	}
+	case EPresenceState::Away:
+	{
+		return EOS_Presence_EStatus::EOS_PS_Away;
+	}
+	case EPresenceState::ExtendedAway:
+	{
+		return EOS_Presence_EStatus::EOS_PS_ExtendedAway;
+	}
+	case EPresenceState::DoNotDisturb:
+	{
+		return EOS_Presence_EStatus::EOS_PS_DoNotDisturb;
+	}
+	}
+	return EOS_Presence_EStatus::EOS_PS_Offline;
+}
+
+FPresenceEOS::FPresenceEOS(FOnlineServicesEOS& InServices)
+	: FPresenceCommon(InServices)
+{
+}
+
+void FPresenceEOS::Initialize()
+{
+	FPresenceCommon::Initialize();
+
+	PresenceHandle = EOS_Platform_GetPresenceInterface(static_cast<FOnlineServicesEOS&>(GetServices()).GetEOSPlatformHandle());
+	check(PresenceHandle != nullptr);
+
+	// Register for friend updates
+	EOS_Presence_AddNotifyOnPresenceChangedOptions Options = { };
+	Options.ApiVersion = EOS_PRESENCE_ADDNOTIFYONPRESENCECHANGED_API_LATEST;
+	NotifyPresenceChangedNotificationId = EOS_Presence_AddNotifyOnPresenceChanged(PresenceHandle, &Options, this, [](const EOS_Presence_PresenceChangedCallbackInfo* Data)
+	{
+		FPresenceEOS* This = reinterpret_cast<FPresenceEOS*>(Data->ClientData);
+		const FOnlineAccountIdHandle LocalUserId = FindAccountIdChecked(Data->LocalUserId);
+
+		This->Services.Get<FAuthEOS>()->ResolveAccountId(LocalUserId, Data->PresenceUserId)
+		.Next([This, LocalUserId](const FOnlineAccountIdHandle& PresenceUserId)
+		{
+			UE_LOG(LogTemp, Warning, TEXT("OnEOSPresenceUpdate: LocalUserId=[%s] PresenceUserId=[%s]"), *ToLogString(LocalUserId), *ToLogString(PresenceUserId));
+			This->UpdateUserPresence(LocalUserId, PresenceUserId);
+		});
+	});
+}
+
+void FPresenceEOS::PreShutdown()
+{
+	EOS_Presence_RemoveNotifyOnPresenceChanged(PresenceHandle, NotifyPresenceChangedNotificationId);
+}
+
+TOnlineAsyncOpHandle<FQueryPresence> FPresenceEOS::QueryPresence(FQueryPresence::Params&& InParams)
+{
+	TOnlineAsyncOpRef<FQueryPresence> Op = GetJoinableOp<FQueryPresence>(MoveTemp(InParams));
+	if (!Op->IsReady())
+	{
+		// Initialize
+		const FQueryPresence::Params& Params = Op->GetParams();
+		if (!Services.Get<FAuthEOS>()->IsLoggedIn(Params.LocalUserId))
+		{
+			// TODO: Error codes
+			Op->SetError(Errors::Unknown());
+			return Op->GetHandle();
+		}
+		EOS_EpicAccountId TargetUserEasId = GetEpicAccountId(Params.TargetUserId);
+		if (!EOS_EpicAccountId_IsValid(TargetUserEasId))
+		{
+			// TODO: Error codes
+			Op->SetError(Errors::Unknown());
+			return Op->GetHandle();
+		}
+
+		// TODO:  If we try to query a local user's presence, is that an error, should we return the cached state, should we still ask EOS?
+		const bool bIsLocalUser = Services.Get<FAuthEOS>()->IsLoggedIn(Params.TargetUserId);
+		if (bIsLocalUser)
+		{
+			Op->SetError(Errors::Unknown()); // TODO: Error codes
+			return Op->GetHandle();
+		}
+
+		Op->Then([this](TOnlineAsyncOp<FQueryPresence>& InAsyncOp) mutable
+			{
+				const FQueryPresence::Params& Params = InAsyncOp.GetParams();
+				EOS_Presence_QueryPresenceOptions QueryPresenceOptions = { };
+				QueryPresenceOptions.ApiVersion = EOS_PRESENCE_QUERYPRESENCE_API_LATEST;
+				QueryPresenceOptions.LocalUserId = GetEpicAccountIdChecked(Params.LocalUserId);
+				QueryPresenceOptions.TargetUserId = GetEpicAccountIdChecked(Params.TargetUserId);
+
+				return EOS_Async<EOS_Presence_QueryPresenceCallbackInfo>(EOS_Presence_QueryPresence, PresenceHandle, QueryPresenceOptions);
+			})
+			.Then([this](TOnlineAsyncOp<FQueryPresence>& InAsyncOp, const EOS_Presence_QueryPresenceCallbackInfo* Data) mutable
+			{
+				UE_LOG(LogTemp, Warning, TEXT("QueryPresenceResult: [%s]"), *LexToString(Data->ResultCode));
+
+				if (Data->ResultCode == EOS_EResult::EOS_Success)
+				{
+					const FQueryPresence::Params& Params = InAsyncOp.GetParams();
+					UpdateUserPresence(Params.LocalUserId, Params.TargetUserId);
+					FQueryPresence::Result Result = { FindOrCreatePresence(Params.LocalUserId, Params.TargetUserId) };
+					InAsyncOp.SetResult(MoveTemp(Result));
+				}
+				else
+				{
+					InAsyncOp.SetError(Errors::Unknown()); // TODO: Error codes
+				}
+			})
+			.Enqueue(GetSerialQueue());
+	}
+	return Op->GetHandle();
+}
+
+TOnlineResult<FGetPresence> FPresenceEOS::GetPresence(FGetPresence::Params&& Params)
+{
+	if (TMap<FOnlineAccountIdHandle, TSharedRef<FUserPresence>>* PresenceList = PresenceLists.Find(Params.LocalUserId))
+	{
+		TSharedRef<FUserPresence>* PresencePtr = PresenceList->Find(Params.TargetUserId);
+		if (PresencePtr)
+		{
+			FGetPresence::Result Result = { *PresencePtr };
+			return TOnlineResult<FGetPresence>(MoveTemp(Result));
+		}
+	}
+	return TOnlineResult<FGetPresence>(Errors::Unknown()); // TODO: error codes
+}
+
+TOnlineAsyncOpHandle<FUpdatePresence> FPresenceEOS::UpdatePresence(FUpdatePresence::Params&& InParams)
+{
+	// TODO: Validate params
+	// EOS_PRESENCE_DATA_MAX_KEYS - number of total keys.  Compare proposed with existing with pending ops?  Include removed keys!
+	// EOS_PRESENCE_DATA_MAX_KEY_LENGTH - length of each key.  Compare updated with this.
+	// EOS_PRESENCE_DATA_MAX_VALUE_LENGTH - length of each value.  Compare updated with this.
+	// EOS_PRESENCE_RICH_TEXT_MAX_VALUE_LENGTH - length of status. Compare updated with this.
+
+	TOnlineAsyncOpRef<FUpdatePresence> Op = GetMergeableOp<FUpdatePresence>(MoveTemp(InParams));
+	if (!Op->IsReady())
+	{
+		// Initialize
+		const FUpdatePresence::Params& Params = Op->GetParams();
+		if (!Services.Get<FAuthEOS>()->IsLoggedIn(Params.LocalUserId))
+		{
+			// TODO: Error codes
+			Op->SetError(Errors::Unknown());
+			return Op->GetHandle();
+		}
+
+		// Don't cache anything from Params as they could be modified by another merge in the meanwhile.
+		Op->Then([this](TOnlineAsyncOp<FUpdatePresence>& InAsyncOp) mutable
+			{
+				const FUpdatePresence::Params& Params = InAsyncOp.GetParams();
+				EOS_HPresenceModification ChangeHandle = nullptr;
+				EOS_Presence_CreatePresenceModificationOptions Options = { };
+				Options.ApiVersion = EOS_PRESENCE_CREATEPRESENCEMODIFICATION_API_LATEST;
+				Options.LocalUserId = GetEpicAccountIdChecked(Params.LocalUserId);
+				EOS_EResult CreatePresenceModificationResult = EOS_Presence_CreatePresenceModification(PresenceHandle, &Options, &ChangeHandle);
+				if (CreatePresenceModificationResult == EOS_EResult::EOS_Success)
+				{
+					check(ChangeHandle != nullptr);
+
+					// State
+					if (Params.Mutations.State.IsSet())
+					{
+						EOS_PresenceModification_SetStatusOptions StatusOptions = { };
+						StatusOptions.ApiVersion = EOS_PRESENCEMODIFICATION_SETSTATUS_API_LATEST;
+						StatusOptions.Status = ToEOS_Presence_EStatus(Params.Mutations.State.GetValue());
+						EOS_EResult SetStatusResult = EOS_PresenceModification_SetStatus(ChangeHandle, &StatusOptions);
+						if (SetStatusResult != EOS_EResult::EOS_Success)
+						{
+							UE_LOG(LogTemp, Warning, TEXT("UpdatePresence: EOS_PresenceModification_SetStatus failed with result %s"), *LexToString(SetStatusResult));
+							InAsyncOp.SetError(Errors::Unknown());
+							return MakeFulfilledPromise<const EOS_Presence_SetPresenceCallbackInfo*>(nullptr).GetFuture();
+						}
+						else
+						{
+							UE_LOG(LogTemp, Warning, TEXT("UpdatePresence: Status set to %u"), (uint8)Params.Mutations.State.GetValue()); // Temp logging
+						}
+					}
+
+					// Raw rich text
+					if (Params.Mutations.StatusString.IsSet())
+					{
+						// Convert the status string as the rich text string
+						EOS_PresenceModification_SetRawRichTextOptions RawRichTextOptions = { };
+						RawRichTextOptions.ApiVersion = EOS_PRESENCEMODIFICATION_SETRAWRICHTEXT_API_LATEST;
+
+						FTCHARToUTF8 Utf8RawRichText(*Params.Mutations.StatusString.GetValue());
+						RawRichTextOptions.RichText = Utf8RawRichText.Get();
+
+						EOS_EResult SetRawRichTextResult = EOS_PresenceModification_SetRawRichText(ChangeHandle, &RawRichTextOptions);
+						if (SetRawRichTextResult != EOS_EResult::EOS_Success)
+						{
+							UE_LOG(LogTemp, Warning, TEXT("UpdatePresence: EOS_PresenceModification_SetRawRichText failed with result %s"), *LexToString(SetRawRichTextResult));
+							InAsyncOp.SetError(Errors::Unknown());
+							return MakeFulfilledPromise<const EOS_Presence_SetPresenceCallbackInfo*>(nullptr).GetFuture();
+						}
+						else
+						{
+							UE_LOG(LogTemp, Warning, TEXT("UpdatePresence: RichText set to %s"), *Params.Mutations.StatusString.GetValue()); // Temp logging
+						}
+					}
+
+					// Removed fields
+					if (Params.Mutations.RemovedProperties.Num() > 0)
+					{
+						// EOS_PresenceModification_DeleteData
+						TArray<FTCHARToUTF8, TInlineAllocator<EOS_PRESENCE_DATA_MAX_KEYS>> Utf8Strings;
+						TArray<EOS_PresenceModification_DataRecordId, TInlineAllocator<EOS_PRESENCE_DATA_MAX_KEYS>> RecordIds;
+						int32 CurrentIndex = 0;
+						for (const FString& RemovedProperty : Params.Mutations.RemovedProperties)
+						{
+							const FTCHARToUTF8& Utf8Key = Utf8Strings.Emplace_GetRef(RemovedProperty);
+
+							EOS_PresenceModification_DataRecordId& RecordId = RecordIds.Emplace_GetRef();
+							RecordId.ApiVersion = EOS_PRESENCEMODIFICATION_DATARECORDID_API_LATEST;
+							RecordId.Key = Utf8Key.Get();
+
+							UE_LOG(LogTemp, Warning, TEXT("UpdatePresence: Removing field %s"), *RemovedProperty); // Temp logging
+						}
+
+						EOS_PresenceModification_DeleteDataOptions DataOptions = { };
+						DataOptions.ApiVersion = EOS_PRESENCEMODIFICATION_DELETEDATA_API_LATEST;
+						DataOptions.RecordsCount = RecordIds.Num();
+						DataOptions.Records = RecordIds.GetData();
+						EOS_EResult DeleteDataResult = EOS_PresenceModification_DeleteData(ChangeHandle, &DataOptions);
+						if (DeleteDataResult != EOS_EResult::EOS_Success)
+						{
+							UE_LOG(LogTemp, Warning, TEXT("UpdatePresence: EOS_PresenceModification_DeleteDataOptions failed with result %s"), *LexToString(DeleteDataResult));
+							InAsyncOp.SetError(Errors::Unknown());
+							return MakeFulfilledPromise<const EOS_Presence_SetPresenceCallbackInfo*>(nullptr).GetFuture();
+						}
+					}
+
+					// Added/Updated fields
+					if (Params.Mutations.UpdatedProperties.Num() > 0)
+					{
+						if (Params.Mutations.UpdatedProperties.Num() > EOS_PRESENCE_DATA_MAX_KEYS)
+						{
+							// TODO: Move this check higher.  Needs to take into account number of present fields (not just ones updated) and removed fields.
+							UE_LOG(LogTemp, Warning, TEXT("UpdatePresence: Too many presence keys.  %u/%u"), Params.Mutations.UpdatedProperties.Num(), EOS_PRESENCE_DATA_MAX_KEYS);
+							InAsyncOp.SetError(Errors::Unknown());
+							return MakeFulfilledPromise<const EOS_Presence_SetPresenceCallbackInfo*>(nullptr).GetFuture();
+						}
+						TArray<FTCHARToUTF8, TInlineAllocator<EOS_PRESENCE_DATA_MAX_KEYS * 2>> Utf8Strings;
+						TArray<EOS_Presence_DataRecord, TInlineAllocator<EOS_PRESENCE_DATA_MAX_KEYS>> Records;
+						int32 CurrentIndex = 0;
+						for (const TPair<FString, FPresenceVariant>& UpdatedProperty : Params.Mutations.UpdatedProperties)
+						{
+							const FTCHARToUTF8& Utf8Key = Utf8Strings.Emplace_GetRef(UpdatedProperty.Key);
+							const FTCHARToUTF8& Utf8Value = Utf8Strings.Emplace_GetRef(UpdatedProperty.Value); // TODO: Better serialization
+
+							EOS_Presence_DataRecord& Record = Records.Emplace_GetRef();
+							Record.ApiVersion = EOS_PRESENCE_DATARECORD_API_LATEST;
+							Record.Key = Utf8Key.Get();
+							Record.Value = Utf8Value.Get();
+							UE_LOG(LogTemp, Warning, TEXT("UpdatePresence: Set field [%s] to [%s]"), *UpdatedProperty.Key, *UpdatedProperty.Value); // Temp logging
+						}
+
+						EOS_PresenceModification_SetDataOptions DataOptions = { };
+						DataOptions.ApiVersion = EOS_PRESENCEMODIFICATION_SETDATA_API_LATEST;
+						DataOptions.RecordsCount = Records.Num();
+						DataOptions.Records = Records.GetData();
+						EOS_EResult SetDataResult = EOS_PresenceModification_SetData(ChangeHandle, &DataOptions);
+						if (SetDataResult != EOS_EResult::EOS_Success)
+						{
+							UE_LOG(LogTemp, Warning, TEXT("UpdatePresence: EOS_PresenceModification_SetData failed with result %s"), *LexToString(SetDataResult));
+							InAsyncOp.SetError(Errors::Unknown());
+							return MakeFulfilledPromise<const EOS_Presence_SetPresenceCallbackInfo*>(nullptr).GetFuture();
+						}
+					}
+
+					EOS_Presence_SetPresenceOptions SetPresenceOptions = { };
+					SetPresenceOptions.ApiVersion = EOS_PRESENCE_SETPRESENCE_API_LATEST;
+					SetPresenceOptions.LocalUserId = GetEpicAccountIdChecked(Params.LocalUserId);
+					SetPresenceOptions.PresenceModificationHandle = ChangeHandle;
+					TFuture<const EOS_Presence_SetPresenceCallbackInfo*> EOSAsyncFuture = EOS_Async<EOS_Presence_SetPresenceCallbackInfo>(EOS_Presence_SetPresence, PresenceHandle, SetPresenceOptions);
+					EOS_PresenceModification_Release(ChangeHandle);
+					return EOSAsyncFuture;
+				}
+				else
+				{
+					// TODO:  Error code
+					InAsyncOp.SetError(Errors::Unknown());
+					return MakeFulfilledPromise<const EOS_Presence_SetPresenceCallbackInfo*>(nullptr).GetFuture();
+				}
+			})
+			.Then([this](TOnlineAsyncOp<FUpdatePresence>& InAsyncOp, const EOS_Presence_SetPresenceCallbackInfo* Data) mutable
+			{
+				UE_LOG(LogTemp, Warning, TEXT("SetPresenceResult: [%s]"), *LexToString(Data->ResultCode));
+
+				if (Data->ResultCode == EOS_EResult::EOS_Success)
+				{
+					// Update local presence
+					const FUpdatePresence::Params& Params = InAsyncOp.GetParams();
+					TSharedRef<FUserPresence> LocalUserPresence = FindOrCreatePresence(Params.LocalUserId, Params.LocalUserId);
+					if (Params.Mutations.State.IsSet())
+					{
+						LocalUserPresence->State = Params.Mutations.State.GetValue();
+					}
+					if (Params.Mutations.StatusString.IsSet())
+					{
+						LocalUserPresence->StatusString = Params.Mutations.StatusString.GetValue();
+					}
+					for (const FString& RemovedKey : Params.Mutations.RemovedProperties)
+					{
+						LocalUserPresence->Properties.Remove(RemovedKey);
+					}
+					for (const TPair<FString, FPresenceVariant>& UpdatedProperty : Params.Mutations.UpdatedProperties)
+					{
+						LocalUserPresence->Properties.Emplace(UpdatedProperty.Key, UpdatedProperty.Value);
+					}
+
+					InAsyncOp.SetResult(FUpdatePresence::Result());
+				}
+				else
+				{
+					InAsyncOp.SetError(Errors::Unknown()); // TODO: Error codes
+				}
+			})
+			.Enqueue(GetSerialQueue());
+	}
+	return Op->GetHandle();
+}
+
+/** Get a user's presence, creating entries if missing */
+TSharedRef<FUserPresence> FPresenceEOS::FindOrCreatePresence(FOnlineAccountIdHandle LocalUserId, FOnlineAccountIdHandle PresenceUserId)
+{
+	TMap<FOnlineAccountIdHandle, TSharedRef<FUserPresence>>& LocalUserPresenceList = PresenceLists.FindOrAdd(LocalUserId);
+	if (const TSharedRef<FUserPresence>* const ExistingPresence = LocalUserPresenceList.Find(PresenceUserId))
+	{
+		return *ExistingPresence;
+	}
+
+	TSharedRef<FUserPresence> UserPresence = MakeShared<FUserPresence>();
+	UserPresence->UserId = PresenceUserId;
+	LocalUserPresenceList.Emplace(PresenceUserId, UserPresence);
+	return UserPresence;
+}
+
+void FPresenceEOS::UpdateUserPresence(FOnlineAccountIdHandle LocalUserId, FOnlineAccountIdHandle PresenceUserId)
+{
+	bool bPresenceHasChanged = false;
+	TSharedRef<FUserPresence> UserPresence = FindOrCreatePresence(LocalUserId, PresenceUserId);
+	// TODO:  Handle updates for local users.  Don't want to conflict with UpdatePresence calls
+
+	// Get presence from EOS
+	EOS_Presence_Info* PresenceInfo = nullptr;
+	EOS_Presence_CopyPresenceOptions Options = { };
+	Options.ApiVersion = EOS_PRESENCE_COPYPRESENCE_API_LATEST;
+	Options.LocalUserId = GetEpicAccountIdChecked(LocalUserId);
+	Options.TargetUserId = GetEpicAccountIdChecked(PresenceUserId);
+	EOS_EResult CopyPresenceResult = EOS_Presence_CopyPresence(PresenceHandle, &Options, &PresenceInfo);
+	if (CopyPresenceResult == EOS_EResult::EOS_Success)
+	{
+		// Convert the presence data to our format
+		EPresenceState NewPresenceState = ToEPresenceState(PresenceInfo->Status);
+		if (UserPresence->State != NewPresenceState)
+		{
+			bPresenceHasChanged = true;
+			UserPresence->State = NewPresenceState;
+		}
+
+		FString NewStatusString = UTF8_TO_TCHAR(PresenceInfo->RichText);
+		if (UserPresence->StatusString != NewStatusString)
+		{
+			bPresenceHasChanged = true;
+			UserPresence->StatusString = MoveTemp(NewStatusString);
+		}
+
+		if (PresenceInfo->Records)
+		{
+			// TODO:  Handle Properties that aren't replicated through presence (eg "ProductId")
+			TArrayView<const EOS_Presence_DataRecord> Records(PresenceInfo->Records, PresenceInfo->RecordsCount);
+			// Detect removals
+			TArray<FString, TInlineAllocator<EOS_PRESENCE_DATA_MAX_KEYS>> RemovedKeys;
+			UserPresence->Properties.GenerateKeyArray(RemovedKeys);
+			for (const EOS_Presence_DataRecord& Record : Records)
+			{
+				FString RecordKey = UTF8_TO_TCHAR(Record.Key);
+				FString RecordValue = UTF8_TO_TCHAR(Record.Value);
+				RemovedKeys.Remove(RecordKey);
+				if (FString* ExistingValue = UserPresence->Properties.Find(RecordKey))
+				{
+					if (*ExistingValue == RecordValue)
+					{
+						continue; // No change
+					}
+				}
+				bPresenceHasChanged = true;
+				UserPresence->Properties.Add(MoveTemp(RecordKey), MoveTemp(RecordValue));
+			}
+			// Any fields that have been removed
+			if (RemovedKeys.Num() > 0)
+			{
+				bPresenceHasChanged = true;
+				for (const FString& RemovedKey : RemovedKeys)
+				{
+					UserPresence->Properties.Remove(RemovedKey);
+				}
+			}
+		}
+		else if (UserPresence->Properties.Num() > 0)
+		{
+			bPresenceHasChanged = true;
+			UserPresence->Properties.Reset();
+		}
+		EOS_Presence_Info_Release(PresenceInfo);
+	}
+	else
+	{
+		UE_LOG(LogTemp, Error, TEXT("UpdateUserPresence: CopyPresence Failed %s"), *LexToString(CopyPresenceResult));
+	}
+
+	if (bPresenceHasChanged)
+	{
+		FPresenceUpdated PresenceUpdatedParams = { LocalUserId, UserPresence };
+		OnPresenceUpdatedEvent.Broadcast(PresenceUpdatedParams);
+	}
+}
+
+/* UE::Online */ }
diff --git a/Engine/Plugins/Online/Experimental/OnlineServicesEOS/Source/Private/OnlineServicesEOSModule.cpp b/Engine/Plugins/Online/Experimental/OnlineServicesEOS/Source/Private/OnlineServicesEOSModule.cpp
--- a/Engine/Plugins/Online/Experimental/OnlineServicesEOS/Source/Private/OnlineServicesEOSModule.cpp
+++ b/Engine/Plugins/Online/Experimental/OnlineServicesEOS/Source/Private/OnlineServicesEOSModule.cpp
@@ -5,8 +5,8 @@
 #include "Modules/ModuleManager.h"
 
 #include "Online/OnlineServicesRegistry.h"
-#include "OnlineServicesEOS.h"
-#include "OnlineIdEOS.h"
+#include "Online/OnlineServicesEOS.h"
+#include "Online/OnlineIdEOS.h"
 
 namespace UE::Online
 {
diff --git a/Engine/Plugins/Online/Experimental/OnlineServicesEOS/Source/Public/Online/ExternalUIEOS.h b/Engine/Plugins/Online/Experimental/OnlineServicesEOS/Source/Public/Online/ExternalUIEOS.h
new file mode 100644
--- /dev/null
+++ b/Engine/Plugins/Online/Experimental/OnlineServicesEOS/Source/Public/Online/ExternalUIEOS.h
@@ -0,0 +1,51 @@
+// Copyright Epic Games, Inc. All Rights Reserved.
+
+#pragma once
+
+#include "CoreMinimal.h"
+#include "Online/ExternalUICommon.h"
+
+#if defined(EOS_PLATFORM_BASE_FILE_NAME)
+#include EOS_PLATFORM_BASE_FILE_NAME
+#endif
+#include "eos_ui_types.h"
+
+namespace UE::Online {
+
+struct FExternalUIEOSConfig
+{
+	/** Is ShowLoginUI enabled? */
+	bool bShowLoginUIEnabled = true;
+};
+
+class FOnlineServicesEOS;
+
+class ONLINESERVICESEOS_API FExternalUIEOS : public FExternalUICommon
+{
+public:
+	using Super = FExternalUICommon;
+
+	FExternalUIEOS(FOnlineServicesEOS& InOwningSubsystem);
+	virtual void Initialize() override;
+	virtual void LoadConfig() override;
+	virtual void PreShutdown() override;
+	virtual TOnlineAsyncOpHandle<FExternalUIShowLoginUI> ShowLoginUI(FExternalUIShowLoginUI::Params&& Params) override;
+	virtual TOnlineAsyncOpHandle<FExternalUIShowFriendsUI> ShowFriendsUI(FExternalUIShowFriendsUI::Params&& Params) override;
+
+protected:
+	using Super::LoadConfig;
+	/** Handle to EOS UI */
+	EOS_HUI UIHandle;
+	/** Config */
+	FExternalUIEOSConfig Config;
+};
+
+namespace Meta {
+
+BEGIN_ONLINE_STRUCT_META(FExternalUIEOSConfig)
+	ONLINE_STRUCT_FIELD(FExternalUIEOSConfig, bShowLoginUIEnabled)
+END_ONLINE_STRUCT_META()
+
+/* Meta */ }
+
+/* UE::Online */ }
diff --git a/Engine/Plugins/Online/Experimental/OnlineServicesEOS/Source/Public/Online/FriendsEOS.h b/Engine/Plugins/Online/Experimental/OnlineServicesEOS/Source/Public/Online/FriendsEOS.h
new file mode 100644
--- /dev/null
+++ b/Engine/Plugins/Online/Experimental/OnlineServicesEOS/Source/Public/Online/FriendsEOS.h
@@ -0,0 +1,38 @@
+// Copyright Epic Games, Inc. All Rights Reserved.
+
+#pragma once
+
+#include "CoreMinimal.h"
+#include "Online/FriendsCommon.h"
+
+#if defined(EOS_PLATFORM_BASE_FILE_NAME)
+#include EOS_PLATFORM_BASE_FILE_NAME
+#endif
+#include "eos_friends_types.h"
+
+namespace UE::Online {
+
+class FOnlineServicesEOS;
+	
+class FFriendsEOS : public FFriendsCommon
+{
+public:
+	FFriendsEOS(FOnlineServicesEOS& InServices);
+
+	virtual void Initialize() override;
+	virtual void PreShutdown() override;
+
+	virtual TOnlineAsyncOpHandle<FQueryFriends> QueryFriends(FQueryFriends::Params&& Params) override;
+	virtual TOnlineResult<FGetFriends> GetFriends(FGetFriends::Params&& Params) override;
+	virtual TOnlineAsyncOpHandle<FAddFriend> AddFriend(FAddFriend::Params&& Params) override;
+
+protected:
+	void OnEOSFriendsUpdate(FOnlineAccountIdHandle LocalUserId, FOnlineAccountIdHandle FriendUserId, EOS_EFriendsStatus PreviousStatus, EOS_EFriendsStatus CurrentStatus);
+protected:
+	TMap<FOnlineAccountIdHandle, TMap<FOnlineAccountIdHandle, TSharedRef<FFriend>>> FriendsLists;
+	EOS_HFriends FriendsHandle = nullptr;
+
+	EOS_NotificationId NotifyFriendsUpdateNotificationId = 0;
+};
+
+/* UE::Online */ }
diff --git a/Engine/Plugins/Online/Experimental/OnlineServicesEOS/Source/Public/Online/PresenceEOS.h b/Engine/Plugins/Online/Experimental/OnlineServicesEOS/Source/Public/Online/PresenceEOS.h
new file mode 100644
--- /dev/null
+++ b/Engine/Plugins/Online/Experimental/OnlineServicesEOS/Source/Public/Online/PresenceEOS.h
@@ -0,0 +1,41 @@
+// Copyright Epic Games, Inc. All Rights Reserved.
+
+#pragma once
+
+#include "CoreMinimal.h"
+#include "Online/PresenceCommon.h"
+
+#if defined(EOS_PLATFORM_BASE_FILE_NAME)
+#include EOS_PLATFORM_BASE_FILE_NAME
+#endif
+#include "eos_presence_types.h"
+
+namespace UE::Online {
+
+class FOnlineServicesEOS;
+	
+class FPresenceEOS : public FPresenceCommon
+{
+public:
+	FPresenceEOS(FOnlineServicesEOS& InServices);
+
+	virtual void Initialize() override;
+	virtual void PreShutdown() override;
+
+	virtual TOnlineAsyncOpHandle<FQueryPresence> QueryPresence(FQueryPresence::Params&& Params) override;
+	virtual TOnlineResult<FGetPresence> GetPresence(FGetPresence::Params&& Params) override;
+	virtual TOnlineAsyncOpHandle<FUpdatePresence> UpdatePresence(FUpdatePresence::Params&& Params) override;
+
+protected:
+	/** Get a user's presence, creating entries if missing */
+	TSharedRef<FUserPresence> FindOrCreatePresence(FOnlineAccountIdHandle LocalUserId, FOnlineAccountIdHandle PresenceUserId);
+	/** Update a user's presence from EOS's current value */
+	void UpdateUserPresence(FOnlineAccountIdHandle LocalUserId, FOnlineAccountIdHandle PresenceUserId);
+protected:
+	EOS_HPresence PresenceHandle = nullptr;
+
+	TMap<FOnlineAccountIdHandle, TMap<FOnlineAccountIdHandle, TSharedRef<FUserPresence>>> PresenceLists;
+	EOS_NotificationId NotifyPresenceChangedNotificationId = 0;
+};
+
+/* UE::Online */ }
