Added Engine checks for automated testing builds

diff --git a/Engine/Plugins/Online/Experimental/OnlineServicesEOS/Source/OnlineServicesEOS.Build.cs b/Engine/Plugins/Online/Experimental/OnlineServicesEOS/Source/OnlineServicesEOS.Build.cs
--- a/Engine/Plugins/Online/Experimental/OnlineServicesEOS/Source/OnlineServicesEOS.Build.cs
+++ b/Engine/Plugins/Online/Experimental/OnlineServicesEOS/Source/OnlineServicesEOS.Build.cs
@@ -6,7 +6,7 @@ using System.IO;
 public class OnlineServicesEOS : ModuleRules
 {
 	public OnlineServicesEOS(ReadOnlyTargetRules Target) : base(Target)
-    {
+	{
 		PublicDependencyModuleNames.AddRange(
 			new string[] {
 				"Core",
@@ -20,11 +20,19 @@ public class OnlineServicesEOS : ModuleRules
 				"CoreOnline",
 				"CoreUObject",
 				"EOSSDK",
-				"EOSShared",
-				"Sockets",
-				"SocketSubsystemEOS",
-				"Engine"
+				"EOSShared"
 			}
 		);
+
+		if (Target.bCompileAgainstEngine)
+		{
+			PrivateDependencyModuleNames.AddRange(
+				new string[]
+				{
+					"SocketSubsystemEOS",
+					"Sockets"
+				}
+			);
+		}
 	}
 }
diff --git a/Engine/Plugins/Online/Experimental/OnlineServicesEOS/Source/Private/Online/OnlineServicesEOS.cpp b/Engine/Plugins/Online/Experimental/OnlineServicesEOS/Source/Private/Online/OnlineServicesEOS.cpp
--- a/Engine/Plugins/Online/Experimental/OnlineServicesEOS/Source/Private/Online/OnlineServicesEOS.cpp
+++ b/Engine/Plugins/Online/Experimental/OnlineServicesEOS/Source/Private/Online/OnlineServicesEOS.cpp
@@ -13,12 +13,15 @@
 #include "Online/OnlineIdEOS.h"
 #include "Online/PresenceEOS.h"
 #include "Online/ExternalUIEOS.h"
+
+#if WITH_ENGINE
 #include "InternetAddrEOS.h"
-#include "Engine/EngineBaseTypes.h"
 #include "NetDriverEOS.h"
+#endif
 
 namespace UE::Online {
 
+#if WITH_ENGINE
 FSocketSubsystemEOSUtils_OnlineServicesEOS::FSocketSubsystemEOSUtils_OnlineServicesEOS(FOnlineServicesEOS* InServicesEOS)
 	: ServicesEOS(InServicesEOS)
 {
@@ -107,6 +110,7 @@ FString FSocketSubsystemEOSUtils_OnlineServicesEOS::GetSessionId()
 
 	return Result;
 }
+#endif
 
 struct FEOSPlatformConfig
 {
@@ -189,6 +193,7 @@ void FOnlineServicesEOS::Initialize()
 	EOSPlatformHandle = SDKManager->CreatePlatform(PlatformOptions);
 	if (EOSPlatformHandle)
 	{
+#if WITH_ENGINE
 		SocketSubsystem = MakeShareable(new FSocketSubsystemEOS(EOSPlatformHandle, MakeShareable(new FSocketSubsystemEOSUtils_OnlineServicesEOS(this))));
 		if (SocketSubsystem)
 		{
@@ -198,6 +203,7 @@ void FOnlineServicesEOS::Initialize()
 				UE_LOG(LogOnlineServices, Verbose, TEXT("[FOnlineServicesEOS::Initialize] Unable to initialize Socket Subsystem. Error=[%s]"), *ErrorStr);
 			}
 		}
+#endif
 	}
 	else
 	{
@@ -219,10 +225,14 @@ TOnlineResult<FGetResolvedConnectString> FOnlineServicesEOS::GetResolvedConnectS
 		{
 			if (Lobby->LobbyId == Params.LobbyId)
 			{
-				//It should look like this: "EOS:0002aeeb5b2d4388a3752dd6d31222ec:GameNetDriver:97"
-				FString NetDriverName = GetDefault<UNetDriverEOS>()->NetDriverName.ToString();
-				FInternetAddrEOS TempAddr(GetProductUserIdChecked(Lobby->OwnerAccountId), NetDriverName, GetTypeHash(NetDriverName));
-				return TOnlineResult<FGetResolvedConnectString>({ TempAddr.ToString(true) });
+#if WITH_ENGINE
+ 				//It should look like this: "EOS:0002aeeb5b2d4388a3752dd6d31222ec:GameNetDriver:97"
+ 				FString NetDriverName = GetDefault<UNetDriverEOS>()->NetDriverName.ToString();
+ 				FInternetAddrEOS TempAddr(GetProductUserIdChecked(Lobby->OwnerAccountId), NetDriverName, GetTypeHash(NetDriverName));
+ 				return TOnlineResult<FGetResolvedConnectString>({ TempAddr.ToString(true) });
+#else
+				return TOnlineResult<FGetResolvedConnectString>(Errors::NotImplemented());
+#endif
 			}
 		}
 		// No matching lobby
diff --git a/Engine/Plugins/Online/Experimental/OnlineServicesEOS/Source/Public/Online/OnlineServicesEOS.h b/Engine/Plugins/Online/Experimental/OnlineServicesEOS/Source/Public/Online/OnlineServicesEOS.h
--- a/Engine/Plugins/Online/Experimental/OnlineServicesEOS/Source/Public/Online/OnlineServicesEOS.h
+++ b/Engine/Plugins/Online/Experimental/OnlineServicesEOS/Source/Public/Online/OnlineServicesEOS.h
@@ -3,7 +3,10 @@
 #pragma once
 
 #include "Online/OnlineServicesCommon.h"
+
+#if WITH_ENGINE
 #include "SocketSubsystemEOS.h"
+#endif
 
 #if defined(EOS_PLATFORM_BASE_FILE_NAME)
 #include EOS_PLATFORM_BASE_FILE_NAME
@@ -16,6 +19,7 @@ namespace UE::Online {
 
 class FOnlineServicesEOS;
 
+#if WITH_ENGINE
 class FSocketSubsystemEOSUtils_OnlineServicesEOS : public ISocketSubsystemEOSUtils
 {
 public:
@@ -28,6 +32,7 @@ public:
 private:
 	FOnlineServicesEOS* ServicesEOS;
 };
+#endif
 
 class ONLINESERVICESEOS_API FOnlineServicesEOS : public FOnlineServicesCommon
 {
@@ -45,7 +50,9 @@ public:
 protected:
 	IEOSPlatformHandlePtr EOSPlatformHandle;
 
+#if WITH_ENGINE
 	TSharedPtr<FSocketSubsystemEOS, ESPMode::ThreadSafe> SocketSubsystem;
+#endif
 };
 
 /* UE::Online */ }
