Add GetBuildUniqueId
Setup LobbiesEOS to use GetBuildUniqueId when creating / searching / joining lobbies.

diff --git a/Engine/Plugins/Online/Experimental/OnlineServices/Source/OnlineServicesCommon/Private/Online/LobbiesCommonTests.cpp b/Engine/Plugins/Online/Experimental/OnlineServices/Source/OnlineServicesCommon/Private/Online/LobbiesCommonTests.cpp
--- a/Engine/Plugins/Online/Experimental/OnlineServices/Source/OnlineServicesCommon/Private/Online/LobbiesCommonTests.cpp
+++ b/Engine/Plugins/Online/Experimental/OnlineServices/Source/OnlineServicesCommon/Private/Online/LobbiesCommonTests.cpp
@@ -1098,14 +1098,14 @@ TFunction<TFuture<void>(TOnlineAsyncOp<OpType>&)> CaptureOperationStepResult(TOn
 }
 
 // todo: TFuture<void> does not work as expected with "Then".
-TFuture<int> AwaitNextGameTick()
+TFuture<int> AwaitSleepFor(double Seconds)
 {
 	TSharedRef<TPromise<int>> Promise = MakeShared<TPromise<int>>();
 	TFuture<int> Future = Promise->GetFuture();
 
 	FTSTicker::GetCoreTicker().AddTicker(
-	TEXT("LobbiesFunctionalTest::AwaitNextGameTick"),
-	0.f,
+	TEXT("LobbiesFunctionalTest::AwaitSleepFor"),
+	Seconds,
 	[Promise](float)
 	{
 		Promise->EmplaceValue();
@@ -1115,6 +1115,12 @@ TFuture<int> AwaitNextGameTick()
 	return Future;
 }
 
+// todo: TFuture<void> does not work as expected with "Then".
+TFuture<int> AwaitNextGameTick()
+{
+	return AwaitSleepFor(0.f);
+}
+
 TFuture<TDefaultErrorResultInternal<FOnlineLobbyIdHandle>> AwaitInvitation(
 	FLobbyEvents& LobbyEvents,
 	FOnlineAccountIdHandle TargetAccountId,
@@ -1408,6 +1414,7 @@ struct FFunctionalTestConfig
 	FString TestAccount2Token;
 
 	float InvitationWaitSeconds = 10.f;
+	float FindMatchReplicationDelay = 5.f;
 };
 
 namespace Meta {
@@ -1433,6 +1440,12 @@ TOnlineAsyncOpHandle<FFunctionalTestLobbies> RunLobbyFunctionalTest(IAuth& AuthI
 	static const FString FindLobbyKeyName = TEXT("FindLobby");
 	static const FString LobbyEventCaptureKeyName = TEXT("LobbyEventCapture");
 	static const FString ConfigNameKeyName = TEXT("Config");
+	static const FString SearchKeyName = TEXT("SearchKey");
+
+	struct FSearchParams
+	{
+		int64 LobbyCreateTime = 0;
+	};
 
 	TSharedRef<FFunctionalTestConfig> TestConfig = MakeShared<FFunctionalTestConfig>();
 	LobbiesCommon.LoadConfig(*TestConfig, ConfigName);
@@ -1759,6 +1772,10 @@ TOnlineAsyncOpHandle<FFunctionalTestLobbies> RunLobbyFunctionalTest(IAuth& AuthI
 	//    Step 5: Leave lobby with primary user.
 	//----------------------------------------------------------------------------------------------
 
+	// todo: additional functional tests for search.
+	// find by lobby id, find by user, find by attributes. Search types are mutually exclusive and
+	// should return invalid args if multiple search types are passed for a search.
+
 	.Then([](TOnlineAsyncOp<FFunctionalTestLobbies>& InAsyncOp)
 	{
 		// Prepare for next test check.
@@ -1766,18 +1783,22 @@ TOnlineAsyncOpHandle<FFunctionalTestLobbies> RunLobbyFunctionalTest(IAuth& AuthI
 
 		const Private::FFunctionalTestLoginUser::Result& User1Info = GetOpDataChecked<Private::FFunctionalTestLoginUser::Result>(InAsyncOp, User1KeyName);
 
+		TSharedRef<FSearchParams> SearchParams = MakeShared<FSearchParams>();
+		SearchParams->LobbyCreateTime = static_cast<int64>(FPlatformTime::Seconds());
+		InAsyncOp.Data.Set(SearchKeyName, TSharedRef<FSearchParams>(SearchParams));
+
 		FCreateLobby::Params Params;
 		Params.LocalUserId = User1Info.AccountInfo->UserId;
 		Params.LocalName = TEXT("test");
 		Params.SchemaName = TEXT("test");
 		Params.MaxMembers = 2;
 		Params.JoinPolicy = ELobbyJoinPolicy::PublicAdvertised;
-		//Params.Attributes;
+		Params.Attributes = {{ TEXT("LobbyCreateTime"), SearchParams->LobbyCreateTime }};
 		Params.LocalUsers.Emplace(FJoinLobbyLocalUserData{User1Info.AccountInfo->UserId, {}});
 		return Params;
 	})
 	.Then(Private::CaptureOperationStepResult<FCreateLobby>(*Op, CreateLobbyKeyName, Private::FBindOperation<FCreateLobby>(&LobbiesCommon, &FLobbiesCommon::CreateLobby)))
-	.Then([](TOnlineAsyncOp<FFunctionalTestLobbies>& InAsyncOp)
+	.Then([TestConfig](TOnlineAsyncOp<FFunctionalTestLobbies>& InAsyncOp)
 	{
 		// Check both lobby joined events were received and in the correct order.
 		const TSharedRef<FLobbyEventCapture>& EventCapture = GetOpDataChecked<TSharedRef<FLobbyEventCapture>>(InAsyncOp, LobbyEventCaptureKeyName);
@@ -1787,16 +1808,23 @@ TOnlineAsyncOpHandle<FFunctionalTestLobbies> RunLobbyFunctionalTest(IAuth& AuthI
 		{
 			InAsyncOp.SetError(Errors::Cancelled());
 		}
+
+		// Wait some time so lobby creation can propagate before searching for it on another client.
+		return Private::AwaitSleepFor(TestConfig->FindMatchReplicationDelay);
+	})
+	.Then([](TOnlineAsyncOp<FFunctionalTestLobbies>& InAsyncOp, int)
+	{
 	})
 	.Then([](TOnlineAsyncOp<FFunctionalTestLobbies>& InAsyncOp)
 	{
 		const Private::FFunctionalTestLoginUser::Result& User2Info = GetOpDataChecked<Private::FFunctionalTestLoginUser::Result>(InAsyncOp, User2KeyName);
 		const FCreateLobby::Result& CreateLobbyResult = GetOpDataChecked<FCreateLobby::Result>(InAsyncOp, CreateLobbyKeyName);
+		const TSharedRef<FSearchParams>& SearchParams = GetOpDataChecked<TSharedRef<FSearchParams>>(InAsyncOp, SearchKeyName);
 
-		// Search for lobby from create
+		// Search for lobby from create. Searching by attribute will also limit results by bucket id.
 		FFindLobbies::Params Params;
 		Params.LocalUserId = User2Info.AccountInfo->UserId;
-		Params.LobbyId = CreateLobbyResult.Lobby->LobbyId;
+		Params.Filters = {{TEXT("LobbyCreateTime"), ELobbyComparisonOp::Equals, SearchParams->LobbyCreateTime}};
 		return Params;
 	})
 	.Then(Private::CaptureOperationStepResult<FFindLobbies>(*Op, FindLobbyKeyName, Private::FBindOperation<FFindLobbies>(&LobbiesCommon, &FLobbiesCommon::FindLobbies)))
diff --git a/Engine/Plugins/Online/Experimental/OnlineServices/Source/OnlineServicesInterface/Private/Online/OnlineServices.cpp b/Engine/Plugins/Online/Experimental/OnlineServices/Source/OnlineServicesInterface/Private/Online/OnlineServices.cpp
--- a/Engine/Plugins/Online/Experimental/OnlineServices/Source/OnlineServicesInterface/Private/Online/OnlineServices.cpp
+++ b/Engine/Plugins/Online/Experimental/OnlineServices/Source/OnlineServicesInterface/Private/Online/OnlineServices.cpp
@@ -3,9 +3,60 @@
 #include "Online/OnlineServices.h"
 
 #include "Online/OnlineServicesRegistry.h"
+#include "HAL/IConsoleManager.h"
+#include "Misc/CommandLine.h"
+#include "Misc/ConfigCacheIni.h"
 
 namespace UE::Online {
 
+int32 GetBuildUniqueId()
+{
+	static bool bStaticCheck = false;
+	static int32 BuildId = 0;
+	static bool bUseBuildIdOverride = false;
+	static int32 BuildIdOverride = 0;
+
+	// add a cvar so it can be modified at runtime
+	static FAutoConsoleVariableRef CVarBuildIdOverride(
+		TEXT("buildidoverride"), BuildId,
+		TEXT("Sets build id used for matchmaking "),
+		ECVF_Default);
+
+	if (!bStaticCheck)
+	{
+		bStaticCheck = true;
+		if (FParse::Value(FCommandLine::Get(), TEXT("BuildIdOverride="), BuildIdOverride) && BuildIdOverride != 0)
+		{
+			bUseBuildIdOverride = true;
+		}
+		else
+		{
+			if (!GConfig->GetBool(TEXT("OnlineServices"), TEXT("bUseBuildIdOverride"), bUseBuildIdOverride, GEngineIni))
+			{
+				UE_LOG(LogTemp, Warning, TEXT("Missing bUseBuildIdOverride= in [OnlineServices] of DefaultEngine.ini"));
+			}
+
+			if (!GConfig->GetInt(TEXT("OnlineServices"), TEXT("BuildIdOverride"), BuildIdOverride, GEngineIni))
+			{
+				UE_LOG(LogTemp, Warning, TEXT("Missing BuildIdOverride= in [OnlineServices] of DefaultEngine.ini"));
+			}
+		}
+
+		if (bUseBuildIdOverride == false)
+		{
+			// Removed old hashing code to use something more predictable and easier to override for when
+			// it's necessary to force compatibility with an older build
+			BuildId = FNetworkVersion::GetNetworkCompatibleChangelist();
+		}
+		else
+		{
+			BuildId = BuildIdOverride;
+		}
+	}
+
+	return BuildId;
+}
+
 bool IsLoaded(EOnlineServices OnlineServices, FName InstanceName)
 {
 	return FOnlineServicesRegistry::Get().IsLoaded(OnlineServices, InstanceName);
diff --git a/Engine/Plugins/Online/Experimental/OnlineServices/Source/OnlineServicesInterface/Public/Online/OnlineServices.h b/Engine/Plugins/Online/Experimental/OnlineServices/Source/OnlineServicesInterface/Public/Online/OnlineServices.h
--- a/Engine/Plugins/Online/Experimental/OnlineServices/Source/OnlineServicesInterface/Public/Online/OnlineServices.h
+++ b/Engine/Plugins/Online/Experimental/OnlineServices/Source/OnlineServicesInterface/Public/Online/OnlineServices.h
@@ -93,6 +93,14 @@ public:
 	virtual TOnlineResult<FGetResolvedConnectString> GetResolvedConnectString(FGetResolvedConnectString::Params&& Params) = 0;
 };
 
+/**
+ * Retrieve the unique id for the executable build. The build id is used for ensuring that client
+ * and server builds are compatible.
+ *
+ * @return The unique id for the executable build.
+ */
+ONLINESERVICESINTERFACE_API int32 GetBuildUniqueId();
+
 /**
  * Check if an instance of the online service is loaded
  *
diff --git a/Engine/Plugins/Online/Experimental/OnlineServicesEOS/Source/Private/Online/LobbiesEOS.cpp b/Engine/Plugins/Online/Experimental/OnlineServicesEOS/Source/Private/Online/LobbiesEOS.cpp
--- a/Engine/Plugins/Online/Experimental/OnlineServicesEOS/Source/Private/Online/LobbiesEOS.cpp
+++ b/Engine/Plugins/Online/Experimental/OnlineServicesEOS/Source/Private/Online/LobbiesEOS.cpp
@@ -40,8 +40,6 @@ void FLobbiesEOS::Initialize()
 	IEOSSDKManager* SDKManager = IEOSSDKManager::Get();
 	check(SDKManager);
 
-	LobbyBucketId = FString::Printf(TEXT("%s_%s"), *SDKManager->GetProductName(), *SDKManager->GetProductVersion());
-
 	EOS_HLobby LobbyInterfaceHandle = EOS_Platform_GetLobbyInterface(static_cast<FOnlineServicesEOS&>(GetServices()).GetEOSPlatformHandle());
 	check(LobbyInterfaceHandle != nullptr);
 
@@ -49,7 +47,8 @@ void FLobbiesEOS::Initialize()
 		LobbyInterfaceHandle,
 		StaticCastSharedPtr<FAuthEOS>(Services.GetAuthInterface()),
 		LobbySchemaRegistry.ToSharedRef(),
-		ServiceSchema.ToSharedRef()});
+		ServiceSchema.ToSharedRef(),
+		{SDKManager->GetProductName(), GetBuildUniqueId()}});
 
 	LobbyDataRegistry = MakeShared<FLobbyDataRegistryEOS>(LobbyPrerequisites.ToSharedRef());
 
@@ -109,7 +108,7 @@ TOnlineAsyncOpHandle<FCreateLobby> FLobbiesEOS::CreateLobby(FCreateLobby::Params
 		// Step 1: Call create lobby.
 
 		const FCreateLobby::Params& Params = InAsyncOp.GetParams();
-		const FTCHARToUTF8 BucketIdUTF8(*LobbyBucketId);
+		const FLobbyBucketIdTranslator<ELobbyTranslationType::ToService> BucketTanslator(LobbyPrerequisites->BucketId);
 
 		// The lobby will be created as invitation only.
 		// Once all local members are joined and the lobby attributes have been set the privacy setting will be moved to the user setting.
@@ -121,7 +120,7 @@ TOnlineAsyncOpHandle<FCreateLobby> FLobbiesEOS::CreateLobby(FCreateLobby::Params
 		CreateLobbyOptions.PermissionLevel = TranslateJoinPolicy(ELobbyJoinPolicy::InvitationOnly);
 		CreateLobbyOptions.bPresenceEnabled = false; // todo: handle
 		CreateLobbyOptions.bAllowInvites = true; // todo: handle
-		CreateLobbyOptions.BucketId = BucketIdUTF8.Get();
+		CreateLobbyOptions.BucketId = BucketTanslator.GetBucketIdEOS();
 		CreateLobbyOptions.bDisableHostMigration = false; // todo: handle
 		CreateLobbyOptions.bEnableRTCRoom = false; // todo: handle
 
@@ -1257,6 +1256,12 @@ TOnlineAsyncOpHandle<FLobbiesEOS::FJoinLobbyMemberImpl> FLobbiesEOS::JoinLobbyMe
 		return Op->GetHandle();
 	}
 
+	if (LobbyDetails->GetInfo()->GetProductVersion() != LobbyPrerequisites->BucketId.GetProductVersion())
+	{
+		Op->SetError(Errors::IncompatibleVersion());
+		return Op->GetHandle();
+	}
+
 	Op->Then([this, LobbyDetails, LobbyData = Params.LobbyData](TOnlineAsyncOp<FJoinLobbyMemberImpl>& InAsyncOp)
 	{
 		const FLobbiesEOS::FJoinLobbyMemberImpl::Params& Params = InAsyncOp.GetParams();
diff --git a/Engine/Plugins/Online/Experimental/OnlineServicesEOS/Source/Private/Online/LobbiesEOSTypes.cpp b/Engine/Plugins/Online/Experimental/OnlineServicesEOS/Source/Private/Online/LobbiesEOSTypes.cpp
--- a/Engine/Plugins/Online/Experimental/OnlineServicesEOS/Source/Private/Online/LobbiesEOSTypes.cpp
+++ b/Engine/Plugins/Online/Experimental/OnlineServicesEOS/Source/Private/Online/LobbiesEOSTypes.cpp
@@ -47,53 +47,21 @@ EOS_EComparisonOp TranslateSearchComparison(ELobbyComparisonOp Op)
 
 } // Private
 
-// Attribute translators.
-enum class EAttributeTranslationType
-{
-	ToService,
-	FromService
-};
-
-template <EAttributeTranslationType>
-class FLobbyAttributeTranslator
-{
-public:
-};
+const FString FLobbyBucketIdEOS::Separator = TEXT("|");
 
-template <>
-class FLobbyAttributeTranslator<EAttributeTranslationType::ToService>
+FLobbyBucketIdEOS::FLobbyBucketIdEOS(FString ProductName, int32 ProductVersion)
+	: ProductName(ProductName.Replace(*Separator, TEXT("_")))
+	, ProductVersion(ProductVersion)
 {
-public:
-	FLobbyAttributeTranslator(const TPair<FLobbyAttributeId, FLobbyVariant>& FromAttributeData);
-	FLobbyAttributeTranslator(FLobbyAttributeId FromAttributeId, const FLobbyVariant& FromAttributeData);
-
-	const EOS_Lobby_AttributeData& GetAttributeData() const { return AttributeData; }
-
-private:
-	FTCHARToUTF8 KeyConverterStorage;
-	TOptional<FTCHARToUTF8> ValueConverterStorage;
-	EOS_Lobby_AttributeData AttributeData;
-};
-
-template <>
-class FLobbyAttributeTranslator<EAttributeTranslationType::FromService>
-{
-public:
-	FLobbyAttributeTranslator(const EOS_Lobby_AttributeData& FromAttributeData);
-
-	const TPair<FLobbyAttributeId, FLobbyVariant>& GetAttributeData() const { return AttributeData; }
-	TPair<FLobbyAttributeId, FLobbyVariant>& GetMutableAttributeData() { return AttributeData; }
-
-private:
-	TPair<FLobbyAttributeId, FLobbyVariant> AttributeData;
-};
+}
 
-FLobbyAttributeTranslator<EAttributeTranslationType::ToService>::FLobbyAttributeTranslator(const TPair<FLobbyAttributeId, FLobbyVariant>& FromAttributeData)
+// Attribute translators.
+FLobbyAttributeTranslator<ELobbyTranslationType::ToService>::FLobbyAttributeTranslator(const TPair<FLobbyAttributeId, FLobbyVariant>& FromAttributeData)
 	: FLobbyAttributeTranslator(FromAttributeData.Key, FromAttributeData.Value)
 {
 }
 
-FLobbyAttributeTranslator<EAttributeTranslationType::ToService>::FLobbyAttributeTranslator(FLobbyAttributeId FromAttributeId, const FLobbyVariant& FromAttributeData)
+FLobbyAttributeTranslator<ELobbyTranslationType::ToService>::FLobbyAttributeTranslator(FLobbyAttributeId FromAttributeId, const FLobbyVariant& FromAttributeData)
 	: KeyConverterStorage(*FromAttributeId.ToString())
 {
 	AttributeData.ApiVersion = EOS_LOBBY_ATTRIBUTEDATA_API_LATEST;
@@ -122,7 +90,7 @@ FLobbyAttributeTranslator<EAttributeTranslationType::ToService>::FLobbyAttribute
 	}
 }
 
-FLobbyAttributeTranslator<EAttributeTranslationType::FromService>::FLobbyAttributeTranslator(const EOS_Lobby_AttributeData& FromAttributeData)
+FLobbyAttributeTranslator<ELobbyTranslationType::FromService>::FLobbyAttributeTranslator(const EOS_Lobby_AttributeData& FromAttributeData)
 {
 	FLobbyAttributeId AttributeId = Private::TranslateLobbyAttributeId(FromAttributeData.Key);
 	FLobbyVariant VariantData;
@@ -153,6 +121,26 @@ FLobbyAttributeTranslator<EAttributeTranslationType::FromService>::FLobbyAttribu
 	AttributeData = TPair<FLobbyAttributeId, FLobbyVariant>(MoveTemp(AttributeId), MoveTemp(VariantData));
 }
 
+FLobbyBucketIdTranslator<ELobbyTranslationType::ToService>::FLobbyBucketIdTranslator(const FLobbyBucketIdEOS& BucketId)
+	: BucketConverterStorage(*FString::Printf(TEXT("%s%s%d"), *BucketId.GetProductName(), *FLobbyBucketIdEOS::Separator, BucketId.GetProductVersion()))
+{
+}
+
+FLobbyBucketIdTranslator<ELobbyTranslationType::FromService>::FLobbyBucketIdTranslator(const char* BucketIdEOS)
+{
+	FUTF8ToTCHAR BucketConverterStorage(BucketIdEOS);
+	FString BucketString(BucketConverterStorage.Get());
+
+	constexpr int32 ExpectedPartsNum = 2;
+	TArray<FString> Parts;
+	if (BucketString.ParseIntoArray(Parts, *FLobbyBucketIdEOS::Separator) == ExpectedPartsNum)
+	{
+		int32 BuildId = 0;
+		::LexFromString(BuildId, *Parts[1]);
+		BucketId = FLobbyBucketIdEOS(Parts[0], BuildId);
+	}
+}
+
 const EOS_HLobbyDetails FLobbyDetailsEOS::InvalidLobbyDetailsHandle = {};
 
 TDefaultErrorResultInternal<TSharedRef<FLobbyDetailsEOS>> FLobbyDetailsEOS::CreateFromLobbyId(
@@ -309,8 +297,8 @@ TFuture<TDefaultErrorResultInternal<TSharedRef<FClientLobbySnapshot>>> FLobbyDet
 		}
 
 		TSharedRef<FClientLobbySnapshot> ClientLobbySnapshot = MakeShared<FClientLobbySnapshot>();
-		ClientLobbySnapshot->MaxMembers = StrongThis->GetInfo()->Get().MaxMembers;
-		ClientLobbySnapshot->JoinPolicy = TranslateJoinPolicy(StrongThis->GetInfo()->Get().PermissionLevel);
+		ClientLobbySnapshot->MaxMembers = StrongThis->GetInfo()->GetMaxMembers();
+		ClientLobbySnapshot->JoinPolicy = TranslateJoinPolicy(StrongThis->GetInfo()->GetPermissionLevel());
 
 		// Resolve member info.
 		{
@@ -358,7 +346,7 @@ TFuture<TDefaultErrorResultInternal<TSharedRef<FClientLobbySnapshot>>> FLobbyDet
 					return;
 				}
 
-				FLobbyAttributeTranslator<EAttributeTranslationType::FromService> AttributeTranslator(*LobbyAttribute->Data);
+				FLobbyAttributeTranslator<ELobbyTranslationType::FromService> AttributeTranslator(*LobbyAttribute->Data);
 				ClientLobbySnapshot->Attributes.Add(MoveTemp(AttributeTranslator.GetMutableAttributeData()));
 			}
 		}
@@ -405,7 +393,7 @@ TDefaultErrorResultInternal<TSharedRef<FClientLobbyMemberSnapshot>> FLobbyDetail
 				return TDefaultErrorResultInternal<TSharedRef<FClientLobbyMemberSnapshot>>(FromEOSError(EOSResult));
 			}
 
-			FLobbyAttributeTranslator<EAttributeTranslationType::FromService> AttributeTranslator(*LobbyAttribute->Data);
+			FLobbyAttributeTranslator<ELobbyTranslationType::FromService> AttributeTranslator(*LobbyAttribute->Data);
 			LobbyMemberSnapshot->Attributes.Add(MoveTemp(AttributeTranslator.GetMutableAttributeData()));
 		}
 	}
@@ -428,7 +416,7 @@ TFuture<EOS_EResult> FLobbyDetailsEOS::ApplyLobbyDataUpdateFromLocalChanges(
 	EOS_Lobby_UpdateLobbyModificationOptions ModificationOptions = {};
 	ModificationOptions.ApiVersion = EOS_LOBBY_UPDATELOBBYMODIFICATION_API_LATEST;
 	ModificationOptions.LocalUserId = GetProductUserIdChecked(LocalUserId);
-	ModificationOptions.LobbyId = GetInfo()->Get().LobbyId;
+	ModificationOptions.LobbyId = GetInfo()->GetLobbyId();
 
 	EOS_EResult EOSResultCode = EOS_Lobby_UpdateLobbyModification(Prerequisites->LobbyInterfaceHandle, &ModificationOptions, &LobbyModificationHandle);
 	if (EOSResultCode != EOS_EResult::EOS_Success)
@@ -455,7 +443,7 @@ TFuture<EOS_EResult> FLobbyDetailsEOS::ApplyLobbyDataUpdateFromLocalChanges(
 	// Add attributes.
 	for (const TPair<FLobbyAttributeId, FLobbyVariant>& MutatedAttribute : Changes.MutatedAttributes)
 	{
-		const FLobbyAttributeTranslator<EAttributeTranslationType::ToService> AttributeTranslator(MutatedAttribute);
+		const FLobbyAttributeTranslator<ELobbyTranslationType::ToService> AttributeTranslator(MutatedAttribute);
 
 		EOS_LobbyModification_AddAttributeOptions AddAttributeOptions = {};
 		AddAttributeOptions.ApiVersion = EOS_LOBBYMODIFICATION_ADDATTRIBUTE_API_LATEST;
@@ -517,7 +505,7 @@ TFuture<EOS_EResult> FLobbyDetailsEOS::ApplyLobbyMemberDataUpdateFromLocalChange
 	EOS_Lobby_UpdateLobbyModificationOptions ModificationOptions = {};
 	ModificationOptions.ApiVersion = EOS_LOBBY_UPDATELOBBYMODIFICATION_API_LATEST;
 	ModificationOptions.LocalUserId = GetProductUserIdChecked(LocalUserId);
-	ModificationOptions.LobbyId = GetInfo()->Get().LobbyId;
+	ModificationOptions.LobbyId = GetInfo()->GetLobbyId();
 
 	EOS_EResult EOSResultCode = EOS_Lobby_UpdateLobbyModification(Prerequisites->LobbyInterfaceHandle, &ModificationOptions, &LobbyModificationHandle);
 	if (EOSResultCode != EOS_EResult::EOS_Success)
@@ -529,7 +517,7 @@ TFuture<EOS_EResult> FLobbyDetailsEOS::ApplyLobbyMemberDataUpdateFromLocalChange
 	// Add member attributes.
 	for (const TPair<FLobbyAttributeId, FLobbyVariant>& MutatedAttribute : Changes.MutatedAttributes)
 	{
-		const FLobbyAttributeTranslator<EAttributeTranslationType::ToService> AttributeTranslator(MutatedAttribute);
+		const FLobbyAttributeTranslator<ELobbyTranslationType::ToService> AttributeTranslator(MutatedAttribute);
 
 		EOS_LobbyModification_AddMemberAttributeOptions AddMemberAttributeOptions = {};
 		AddMemberAttributeOptions.ApiVersion = EOS_LOBBYMODIFICATION_ADDMEMBERATTRIBUTE_API_LATEST;
@@ -605,9 +593,16 @@ TDefaultErrorResultInternal<TSharedRef<FLobbyDetailsInfoEOS>> FLobbyDetailsInfoE
 	return TDefaultErrorResultInternal<TSharedRef<FLobbyDetailsInfoEOS>>(MakeShared<FLobbyDetailsInfoEOS>(FLobbyDetailsInfoPtr(LobbyDetailsInfo)));
 }
 
-FLobbyDetailsInfoEOS::FLobbyDetailsInfoEOS(FLobbyDetailsInfoPtr&& LobbyDetailsInfo)
-	: LobbyDetailsInfo(MoveTempIfPossible(LobbyDetailsInfo))
+FLobbyDetailsInfoEOS::FLobbyDetailsInfoEOS(FLobbyDetailsInfoPtr&& InLobbyDetailsInfo)
+	: LobbyDetailsInfo(MoveTempIfPossible(InLobbyDetailsInfo))
 {
+	const FLobbyBucketIdTranslator<ELobbyTranslationType::FromService> BucketTranslator(LobbyDetailsInfo->BucketId);
+	BucketId = BucketTranslator.GetBucketId();
+
+	if (!BucketId.IsValid())
+	{
+		UE_LOG(LogTemp, Warning, TEXT("[FLobbyDetailsInfoEOS] Failed to parse lobby bucket id. Lobby: %s, Bucket: %s"), UTF8_TO_TCHAR(LobbyDetailsInfo->LobbyId), UTF8_TO_TCHAR(LobbyDetailsInfo->BucketId));
+	}
 }
 
 FLobbyDataEOS::~FLobbyDataEOS()
@@ -660,7 +655,7 @@ FLobbyDataEOS::FLobbyDataEOS(
 	: ClientLobbyData(ClientLobbyData)
 	, LobbyDetailsInfo(LobbyDetailsInfo)
 	, UnregisterFn(MoveTemp(UnregisterFn))
-	, LobbyId(Private::TranslateLobbyId(LobbyDetailsInfo->Get().LobbyId))
+	, LobbyId(Private::TranslateLobbyId(LobbyDetailsInfo->GetLobbyId()))
 {
 }
 
@@ -737,7 +732,7 @@ TSharedPtr<FLobbyDataEOS> FLobbyDataRegistryEOS::Find(FOnlineLobbyIdHandle Lobby
 
 TFuture<TDefaultErrorResultInternal<TSharedRef<FLobbyDataEOS>>> FLobbyDataRegistryEOS::FindOrCreateFromLobbyDetails(FOnlineAccountIdHandle LocalUserId, const TSharedRef<FLobbyDetailsEOS>& LobbyDetails)
 {
-	if (TSharedPtr<FLobbyDataEOS> FindResult = Find(LobbyDetails->GetInfo()->Get().LobbyId))
+	if (TSharedPtr<FLobbyDataEOS> FindResult = Find(LobbyDetails->GetInfo()->GetLobbyId()))
 	{
 		FindResult->AddUserLobbyDetails(LocalUserId, LobbyDetails);
 		return MakeFulfilledPromise<TDefaultErrorResultInternal<TSharedRef<FLobbyDataEOS>>>(FindResult.ToSharedRef()).GetFuture();
@@ -867,7 +862,7 @@ TFuture<TDefaultErrorResultInternal<TSharedRef<FLobbySearchEOS>>> FLobbySearchEO
 		return MakeFulfilledPromise<TDefaultErrorResultInternal<TSharedRef<FLobbySearchEOS>>>(FromEOSError(EOSResult)).GetFuture();
 	}
 
-	if (Params.LobbyId)
+	if (Params.LobbyId) // Search for specific lobby.
 	{
 		TSharedPtr<FLobbyDataEOS> LobbyData = LobbyRegistry->Find(*Params.LobbyId);
 		if (!LobbyData)
@@ -886,8 +881,7 @@ TFuture<TDefaultErrorResultInternal<TSharedRef<FLobbySearchEOS>>> FLobbySearchEO
 			return MakeFulfilledPromise<TDefaultErrorResultInternal<TSharedRef<FLobbySearchEOS>>>(FromEOSError(EOSResult)).GetFuture();
 		}
 	}
-
-	if (Params.TargetUser)
+	else if (Params.TargetUser) // Search for specific user.
 	{
 		EOS_LobbySearch_SetTargetUserIdOptions SetTargetUserIdOptions = {};
 		SetTargetUserIdOptions.ApiVersion = EOS_LOBBYSEARCH_SETTARGETUSERID_API_LATEST;
@@ -900,25 +894,47 @@ TFuture<TDefaultErrorResultInternal<TSharedRef<FLobbySearchEOS>>> FLobbySearchEO
 			return MakeFulfilledPromise<TDefaultErrorResultInternal<TSharedRef<FLobbySearchEOS>>>(FromEOSError(EOSResult)).GetFuture();
 		}
 	}
-
-	for (const FFindLobbySearchFilter& Filter :  Params.Filters)
+	else // Search using parameters.
 	{
-		const FLobbyAttributeTranslator<EAttributeTranslationType::ToService> AttributeTranslator(
-			Filter.AttributeName, Filter.ComparisonValue);
+		// Bucket id.
+		{
+			const FLobbyBucketIdTranslator<ELobbyTranslationType::ToService> BucketTranslator(Prerequisites->BucketId);
 
-		EOS_Lobby_AttributeData AttributeData;
-		AttributeData.ApiVersion = EOS_LOBBY_ATTRIBUTEDATA_API_LATEST;
+			EOS_Lobby_AttributeData AttributeData;
+			AttributeData.ApiVersion = EOS_LOBBY_ATTRIBUTEDATA_API_LATEST;
+			AttributeData.Key = EOS_LOBBY_SEARCH_BUCKET_ID;
+			AttributeData.ValueType = EOS_ELobbyAttributeType::EOS_AT_STRING;
+			AttributeData.Value.AsUtf8 = BucketTranslator.GetBucketIdEOS();
 
-		EOS_LobbySearch_SetParameterOptions SetParameterOptions = {};
-		SetParameterOptions.ApiVersion = EOS_LOBBYSEARCH_SETTARGETUSERID_API_LATEST;
-		SetParameterOptions.Parameter = &AttributeTranslator.GetAttributeData();
-		SetParameterOptions.ComparisonOp = Private::TranslateSearchComparison(Filter.ComparisonOp);
+			EOS_LobbySearch_SetParameterOptions SetParameterOptions = {};
+			SetParameterOptions.ApiVersion = EOS_LOBBYSEARCH_SETTARGETUSERID_API_LATEST;
+			SetParameterOptions.Parameter = &AttributeData;
+			SetParameterOptions.ComparisonOp = EOS_EComparisonOp::EOS_CO_EQUAL;
 
-		EOSResult = EOS_LobbySearch_SetParameter(SearchHandle->Get(), &SetParameterOptions);
-		if (EOSResult != EOS_EResult::EOS_Success)
+			EOSResult = EOS_LobbySearch_SetParameter(SearchHandle->Get(), &SetParameterOptions);
+			if (EOSResult != EOS_EResult::EOS_Success)
+			{
+				// todo: errors
+				return MakeFulfilledPromise<TDefaultErrorResultInternal<TSharedRef<FLobbySearchEOS>>>(FromEOSError(EOSResult)).GetFuture();
+			}
+		}
+
+		for (const FFindLobbySearchFilter& Filter :  Params.Filters)
 		{
-			// todo: errors
-			return MakeFulfilledPromise<TDefaultErrorResultInternal<TSharedRef<FLobbySearchEOS>>>(FromEOSError(EOSResult)).GetFuture();
+			const FLobbyAttributeTranslator<ELobbyTranslationType::ToService> AttributeTranslator(
+				Filter.AttributeName, Filter.ComparisonValue);
+
+			EOS_LobbySearch_SetParameterOptions SetParameterOptions = {};
+			SetParameterOptions.ApiVersion = EOS_LOBBYSEARCH_SETTARGETUSERID_API_LATEST;
+			SetParameterOptions.Parameter = &AttributeTranslator.GetAttributeData();
+			SetParameterOptions.ComparisonOp = Private::TranslateSearchComparison(Filter.ComparisonOp);
+
+			EOSResult = EOS_LobbySearch_SetParameter(SearchHandle->Get(), &SetParameterOptions);
+			if (EOSResult != EOS_EResult::EOS_Success)
+			{
+				// todo: errors
+				return MakeFulfilledPromise<TDefaultErrorResultInternal<TSharedRef<FLobbySearchEOS>>>(FromEOSError(EOSResult)).GetFuture();
+			}
 		}
 	}
 
diff --git a/Engine/Plugins/Online/Experimental/OnlineServicesEOS/Source/Private/Online/LobbiesEOSTypes.h b/Engine/Plugins/Online/Experimental/OnlineServicesEOS/Source/Private/Online/LobbiesEOSTypes.h
--- a/Engine/Plugins/Online/Experimental/OnlineServicesEOS/Source/Private/Online/LobbiesEOSTypes.h
+++ b/Engine/Plugins/Online/Experimental/OnlineServicesEOS/Source/Private/Online/LobbiesEOSTypes.h
@@ -25,6 +25,112 @@ struct FLobbyEvents;
 class FAuthEOS;
 struct FClientLobbyMemberSnapshot;
 
+//--------------------------------------------------------------------------------------------------
+// Translated types
+//--------------------------------------------------------------------------------------------------
+
+class FLobbyBucketIdEOS
+{
+public:
+	static const FString Separator;
+
+	FLobbyBucketIdEOS() = default;
+	FLobbyBucketIdEOS(const FLobbyBucketIdEOS&) = default;
+	FLobbyBucketIdEOS(FLobbyBucketIdEOS&&) = default;
+	FLobbyBucketIdEOS& operator=(const FLobbyBucketIdEOS&) = default;
+	FLobbyBucketIdEOS& operator=(FLobbyBucketIdEOS&&) = default;
+
+	FLobbyBucketIdEOS(FString ProductName, int32 ProductVersion);
+
+	const FString& GetProductName() const { return ProductName; }
+	int32 GetProductVersion() const { return ProductVersion; }
+	bool IsValid() const { return !ProductName.IsEmpty(); }
+
+private:
+	FString ProductName;
+	int32 ProductVersion;
+};
+
+//--------------------------------------------------------------------------------------------------
+// Translators
+//--------------------------------------------------------------------------------------------------
+enum class ELobbyTranslationType
+{
+	ToService,
+	FromService
+};
+
+//--------------------------------------------------------------------------------------------------
+// Lobby attribute
+
+template <ELobbyTranslationType>
+class FLobbyAttributeTranslator
+{
+public:
+};
+
+template <>
+class FLobbyAttributeTranslator<ELobbyTranslationType::ToService>
+{
+public:
+	FLobbyAttributeTranslator(const TPair<FLobbyAttributeId, FLobbyVariant>& FromAttributeData);
+	FLobbyAttributeTranslator(FLobbyAttributeId FromAttributeId, const FLobbyVariant& FromAttributeData);
+
+	const EOS_Lobby_AttributeData& GetAttributeData() const { return AttributeData; }
+
+private:
+	FTCHARToUTF8 KeyConverterStorage;
+	TOptional<FTCHARToUTF8> ValueConverterStorage;
+	EOS_Lobby_AttributeData AttributeData;
+};
+
+template <>
+class FLobbyAttributeTranslator<ELobbyTranslationType::FromService>
+{
+public:
+	FLobbyAttributeTranslator(const EOS_Lobby_AttributeData& FromAttributeData);
+
+	const TPair<FLobbyAttributeId, FLobbyVariant>& GetAttributeData() const { return AttributeData; }
+	TPair<FLobbyAttributeId, FLobbyVariant>& GetMutableAttributeData() { return AttributeData; }
+
+private:
+	TPair<FLobbyAttributeId, FLobbyVariant> AttributeData;
+};
+
+//--------------------------------------------------------------------------------------------------
+// Bucket id
+
+template <ELobbyTranslationType>
+class FLobbyBucketIdTranslator;
+
+template <>
+class FLobbyBucketIdTranslator<ELobbyTranslationType::ToService>
+{
+public:
+	FLobbyBucketIdTranslator(const FLobbyBucketIdEOS& BucketId);
+
+	const char* GetBucketIdEOS() const { return BucketConverterStorage.Get(); }
+
+private:
+	FTCHARToUTF8 BucketConverterStorage;
+};
+
+template <>
+class FLobbyBucketIdTranslator<ELobbyTranslationType::FromService>
+{
+public:
+	FLobbyBucketIdTranslator(const char* BucketIdEOS);
+
+	const FLobbyBucketIdEOS& GetBucketId() const { return BucketId; }
+	FLobbyBucketIdEOS& GetMutableBucketId() { return BucketId; }
+
+private:
+	FLobbyBucketIdEOS BucketId;
+};
+
+//--------------------------------------------------------------------------------------------------
+// Enumerations
+
 inline EOS_ELobbyPermissionLevel TranslateJoinPolicy(ELobbyJoinPolicy JoinPolicy)
 {
 	switch (JoinPolicy)
@@ -47,6 +153,10 @@ inline ELobbyJoinPolicy TranslateJoinPolicy(EOS_ELobbyPermissionLevel JoinPolicy
 	};
 }
 
+//--------------------------------------------------------------------------------------------------
+// Structures.
+//--------------------------------------------------------------------------------------------------
+
 using FLobbySearchParameters = FFindLobbies::Params;
 
 /**
@@ -60,6 +170,7 @@ struct FLobbyPrerequisitesEOS
 	TWeakPtr<FAuthEOS> AuthInterface;
 	TSharedRef<const FLobbySchemaRegistry> SchemaRegistry;
 	TSharedRef<const FLobbySchema> ServiceSchema;
+	FLobbyBucketIdEOS BucketId;
 };
 
 /**
@@ -173,11 +284,16 @@ class FLobbyDetailsInfoEOS final
 {
 public:
 	UE_NONCOPYABLE(FLobbyDetailsInfoEOS);
-	FLobbyDetailsInfoEOS() = default;
+	FLobbyDetailsInfoEOS() = delete;
 
 	static TDefaultErrorResultInternal<TSharedRef<FLobbyDetailsInfoEOS>> Create(EOS_HLobbyDetails LobbyDetailsHandle);
 
-	const EOS_LobbyDetails_Info& Get() const { check(LobbyDetailsInfo); return *LobbyDetailsInfo; }
+	EOS_LobbyId GetLobbyId() const { return LobbyDetailsInfo->LobbyId; }
+	int32 GetMaxMembers() const { return LobbyDetailsInfo->MaxMembers; }
+	EOS_ELobbyPermissionLevel GetPermissionLevel() const { return LobbyDetailsInfo->PermissionLevel; }
+
+	const FString& GetProductName() const { return BucketId.GetProductName(); }
+	int32 GetProductVersion() const { return BucketId.GetProductVersion(); }
 
 private:
 	struct FEOSLobbyDetailsInfoDeleter
@@ -195,6 +311,7 @@ private:
 	FLobbyDetailsInfoEOS(FLobbyDetailsInfoPtr&& LobbyDetailsInfo);
 
 	TUniquePtr<EOS_LobbyDetails_Info, FEOSLobbyDetailsInfoDeleter> LobbyDetailsInfo;
+	FLobbyBucketIdEOS BucketId;
 };
 
 /** Lobby data is the bookkeeping object for a lobby. It contains the client-side representation of a lobby. */
@@ -214,7 +331,7 @@ public:
 
 	FOnlineLobbyIdHandle GetLobbyIdHandle() const { return ClientLobbyData->GetPublicData().LobbyId; }
 	const TSharedRef<FClientLobbyData>& GetClientLobbyData() const { return ClientLobbyData; }
-	EOS_LobbyId GetLobbyIdEOS() const { return LobbyDetailsInfo->Get().LobbyId; }
+	EOS_LobbyId GetLobbyIdEOS() const { return LobbyDetailsInfo->GetLobbyId(); }
 	const FString& GetLobbyId() const { return LobbyId; }
 
 	void AddUserLobbyDetails(FOnlineAccountIdHandle LocalUserId, const TSharedPtr<FLobbyDetailsEOS>& LobbyDetails);
diff --git a/Engine/Plugins/Online/Experimental/OnlineServicesEOS/Source/Public/Online/LobbiesEOS.h b/Engine/Plugins/Online/Experimental/OnlineServicesEOS/Source/Public/Online/LobbiesEOS.h
--- a/Engine/Plugins/Online/Experimental/OnlineServicesEOS/Source/Public/Online/LobbiesEOS.h
+++ b/Engine/Plugins/Online/Experimental/OnlineServicesEOS/Source/Public/Online/LobbiesEOS.h
@@ -323,8 +323,6 @@ protected:
 
 	TOnlineAsyncOpHandle<FProcessLobbyNotificationImpl> ProcessLobbyNotificationImplOp(FProcessLobbyNotificationImpl::Params&& Params);
 
-	FString LobbyBucketId;
-
 	EOSEventRegistrationPtr OnLobbyUpdatedEOSEventRegistration;
 	EOSEventRegistrationPtr OnLobbyMemberUpdatedEOSEventRegistration;
 	EOSEventRegistrationPtr OnLobbyMemberStatusReceivedEOSEventRegistration;
