- Speculative fix for suspected use after free

diff --git a/Engine/Plugins/Online/Experimental/OnlineServicesEOS/Source/Private/Online/ExternalUIEOS.cpp b/Engine/Plugins/Online/Experimental/OnlineServicesEOS/Source/Private/Online/ExternalUIEOS.cpp
--- a/Engine/Plugins/Online/Experimental/OnlineServicesEOS/Source/Private/Online/ExternalUIEOS.cpp
+++ b/Engine/Plugins/Online/Experimental/OnlineServicesEOS/Source/Private/Online/ExternalUIEOS.cpp
@@ -53,9 +53,10 @@ TOnlineAsyncOpHandle<FExternalUIShowLoginUI> FExternalUIEOS::ShowLoginUI(FExtern
 			AuthLoginParams.CredentialsType = TEXT("accountportal");
 			AuthLoginParams.Scopes = InAsyncOp.GetParams().Scopes;
 
-			TPromise<void>* TempPromise = new TPromise<void>;
+			TPromise<void>* Promise = new TPromise<void>();
+			TFuture<void> Future = Promise->GetFuture();
 			GetServices().GetAuthInterface()->Login(MoveTemp(AuthLoginParams))
-				.OnComplete([Op = InAsyncOp.AsShared(), TempPromise](const TOnlineResult<FAuthLogin>& AuthLoginResult) mutable
+				.OnComplete([Op = InAsyncOp.AsShared(), Promise](const TOnlineResult<FAuthLogin>& AuthLoginResult) mutable
 				{
 					if (AuthLoginResult.IsOk())
 					{
@@ -66,10 +67,10 @@ TOnlineAsyncOpHandle<FExternalUIShowLoginUI> FExternalUIEOS::ShowLoginUI(FExtern
 					{
 						Op->SetError(CopyTemp(AuthLoginResult.GetErrorValue()));
 					}
-					TempPromise->SetValue();
-					delete TempPromise;
+					Promise->SetValue();
+					delete Promise;
 				});
-			return TempPromise->GetFuture();
+			return Future;
 		}).Enqueue(GetSerialQueue());
 	}
 	else
