Make EAS login an optional step

diff --git a/Engine/Plugins/Online/Experimental/OnlineServicesEOS/Source/Private/Online/AuthEOS.cpp b/Engine/Plugins/Online/Experimental/OnlineServicesEOS/Source/Private/Online/AuthEOS.cpp
--- a/Engine/Plugins/Online/Experimental/OnlineServicesEOS/Source/Private/Online/AuthEOS.cpp
+++ b/Engine/Plugins/Online/Experimental/OnlineServicesEOS/Source/Private/Online/AuthEOS.cpp
@@ -111,71 +111,97 @@ struct FEOSAuthCredentials :
 		return *this;
 	}
 
-	void inline SetExternalCredentialType(EOnlineServices InPlatform, const FString& InCredentialsId)
+	void SetToken(const FCredentialsToken& InToken)
 	{
-		if (InPlatform == EOnlineServices::Steam)
-		{
-			ExternalType = EOS_EExternalCredentialType::EOS_ECT_STEAM_APP_TICKET;
-		}
-		else if (InPlatform == EOnlineServices::PSN)
-		{
-			ExternalType = EOS_EExternalCredentialType::EOS_ECT_PSN_ID_TOKEN;
-		}
-		else if (InPlatform == EOnlineServices::Xbox)
-		{
-			ExternalType = EOS_EExternalCredentialType::EOS_ECT_XBL_XSTS_TOKEN;
-		}
-		else if (InPlatform == EOnlineServices::Nintendo)
-		{
-			if (InCredentialsId == TEXT("NintendoAccount"))
-			{
-				ExternalType = EOS_EExternalCredentialType::EOS_ECT_NINTENDO_ID_TOKEN;
-			}
-			else
-			{
-				ExternalType = EOS_EExternalCredentialType::EOS_ECT_NINTENDO_NSA_ID_TOKEN;
-			}
-		}
-		else if (InPlatform == EOnlineServices::Apple)
+		if (InToken.IsType<TArray<uint8>>())
 		{
-			ExternalType = EOS_EExternalCredentialType::EOS_ECT_APPLE_ID_TOKEN;
+			const TArray<uint8>& TokenData = InToken.Get<TArray<uint8>>();
+			uint32_t InOutBufferLength = EOS_OSS_STRING_BUFFER_LENGTH;
+			EOS_ByteArray_ToString(TokenData.GetData(), TokenData.Num(), TokenAnsi, &InOutBufferLength);
 		}
-		else if (InPlatform == EOnlineServices::Google)
+		else if (InToken.IsType<FString>())
 		{
-			ExternalType = EOS_EExternalCredentialType::EOS_ECT_GOOGLE_ID_TOKEN;
+			const FString& TokenString = InToken.Get<FString>();
+			FCStringAnsi::Strncpy(TokenAnsi, TCHAR_TO_UTF8(*TokenString), EOS_MAX_TOKEN_SIZE);
 		}
 		else
 		{
-			// Unknown means OpenID
-			ExternalType = EOS_EExternalCredentialType::EOS_ECT_OPENID_ACCESS_TOKEN;
+			UE_LOG(LogTemp, Warning, TEXT("SetToken cannot be set with an invalid credentials token parameter. Please ensure there is valid data in the credentials token."));
 		}
 	}
 
+	char IdAnsi[EOS_OSS_STRING_BUFFER_LENGTH];
+	char TokenAnsi[EOS_MAX_TOKEN_SIZE];
+};
+
+class FEOSConnectLoginCredentials : public EOS_Connect_Credentials
+{
+public:
+	FEOSConnectLoginCredentials()
+	{
+		ApiVersion = EOS_CONNECT_CREDENTIALS_API_LATEST;
+	}
+	virtual ~FEOSConnectLoginCredentials() {}
+};
+
+class FEOSConnectLoginCredentialsEOS : public FEOSConnectLoginCredentials
+{
+public:
+	FEOSConnectLoginCredentialsEOS()
+	{
+		Token = TokenAnsi;
+		TokenAnsi[0] = '\0';
+	}
 	void SetToken(const FCredentialsToken& InToken)
 	{
 		if (InToken.IsType<TArray<uint8>>())
 		{
 			const TArray<uint8>& TokenData = InToken.Get<TArray<uint8>>();
-			uint32_t InOutBufferLength = EOS_OSS_STRING_BUFFER_LENGTH;
+			uint32_t InOutBufferLength = EOS_MAX_TOKEN_SIZE;
 			EOS_ByteArray_ToString(TokenData.GetData(), TokenData.Num(), TokenAnsi, &InOutBufferLength);
 		}
 		else if (InToken.IsType<FString>())
 		{
 			const FString& TokenString = InToken.Get<FString>();
-			FCStringAnsi::Strncpy(TokenAnsi, TCHAR_TO_UTF8(*TokenString), TokenString.Len() + 1);
+			FCStringAnsi::Strncpy(TokenAnsi, TCHAR_TO_UTF8(*TokenString), EOS_MAX_TOKEN_SIZE);
 		}
 		else
 		{
-			UE_LOG(LogTemp, Warning, TEXT("FEOSAuthCredentials::Token cannot be set with an invalid credentials token parameter. Please ensure there is valid data in the credentials token."));
+			UE_LOG(LogTemp, Warning, TEXT("SetToken cannot be set with an invalid credentials token parameter. Please ensure there is valid data in the credentials token."));
 		}
 	}
 
-	char IdAnsi[EOS_OSS_STRING_BUFFER_LENGTH];
 	char TokenAnsi[EOS_MAX_TOKEN_SIZE];
 };
 
-FAuthEOS::FAuthEOS(FOnlineServicesEOS& InServices)
+class FEOSConnectLoginCredentialsEAS
+	: public FEOSConnectLoginCredentials
+	, public FNoncopyable
+{
+public:
+	FEOSConnectLoginCredentialsEAS(EOS_Auth_Token* InEASToken)
+		: EASToken(InEASToken)
+	{
+		check(EASToken);
+		Type = EOS_EExternalCredentialType::EOS_ECT_EPIC;
+		Token = EASToken->AccessToken;
+	}
+	virtual ~FEOSConnectLoginCredentialsEAS()
+	{
+		if (EASToken)
+		{
+			EOS_Auth_Token_Release(EASToken);
+		}
+	}
+protected:
+	EOS_Auth_Token* EASToken = nullptr;
+};
+
+
+
+FAuthEOS::FAuthEOS(FOnlineServicesEOS& InServices, bool bInUseEAS)
 	: FAuthCommon(InServices)
+	, bUseEAS(bInUseEAS)
 {
 }
 
@@ -346,10 +372,87 @@ TOnlineAsyncOpHandle<FAuthLogin> FAuthEOS::Login(FAuthLogin::Params&& Params)
 		return Op->GetHandle();
 	}
 
+	TOnlineChainableAsyncOp<FAuthLogin, TSharedPtr<FEOSConnectLoginCredentials>> ConnectLoginOp = [this, &Op]()
+	{
+		if (bUseEAS)
+		{
+			return LoginEAS(*Op);
+		}
+		return Op->Then([this](TOnlineAsyncOp<FAuthLogin>& InAsyncOp)
+		{
+			return MakeConnectLoginCredentials(InAsyncOp);
+		});
+	}();
+	ConnectLoginOp.Then([this](TOnlineAsyncOp<FAuthLogin>& InAsyncOp, TSharedPtr<FEOSConnectLoginCredentials>&& InConnectLoginCredentials)
+	{
+		EOS_Connect_LoginOptions ConnectLoginOptions = { };
+		ConnectLoginOptions.ApiVersion = EOS_CONNECT_LOGIN_API_LATEST;
+		ConnectLoginOptions.Credentials = InConnectLoginCredentials.Get();
+
+		return EOS_Async<EOS_Connect_LoginCallbackInfo>(EOS_Connect_Login, ConnectHandle, ConnectLoginOptions);
+	})
+	.Then([this](TOnlineAsyncOp<FAuthLogin>& InAsyncOp, const EOS_Connect_LoginCallbackInfo* Data)
+	{
+		UE_LOG(LogTemp, Verbose, TEXT("[FAuthEOS::Login] EOS_Connect_Login Result: [%s]"), *LexToString(Data->ResultCode));
+
+		if (Data->ResultCode == EOS_EResult::EOS_Success)
+		{
+			// We cache the Product User Id to use it in later stages of the login process
+			InAsyncOp.Data.Set(TEXT("ProductUserId"), Data->LocalUserId);
+
+			ProcessSuccessfulLogin(InAsyncOp);
+		}
+		else if (Data->ResultCode == EOS_EResult::EOS_InvalidUser && Data->ContinuanceToken != nullptr)
+		{
+			EOS_Connect_CreateUserOptions ConnectCreateUserOptions = { };
+			ConnectCreateUserOptions.ApiVersion = EOS_CONNECT_CREATEUSER_API_LATEST;
+			ConnectCreateUserOptions.ContinuanceToken = Data->ContinuanceToken;
+
+			return EOS_Async<EOS_Connect_CreateUserCallbackInfo>(EOS_Connect_CreateUser, ConnectHandle, ConnectCreateUserOptions);
+		}
+		else
+		{
+			// TODO: EAS Logout
+			InAsyncOp.SetError(Errors::Unknown()); // TODO
+		}
+
+		return MakeFulfilledPromise<const EOS_Connect_CreateUserCallbackInfo*>().GetFuture();
+	})
+	.Then([this](TOnlineAsyncOp<FAuthLogin>& InAsyncOp, const EOS_Connect_CreateUserCallbackInfo* Data)
+	{
+		UE_LOG(LogTemp, Verbose, TEXT("[FAuthEOS::Login] EOS_Connect_CreateUser Result: [%s]"), *LexToString(Data->ResultCode));
+
+		if (Data->ResultCode == EOS_EResult::EOS_Success)
+		{
+			// We cache the Product User Id to use it in later stages of the login process
+			InAsyncOp.Data.Set(TEXT("ProductUserId"), Data->LocalUserId);
+
+			ProcessSuccessfulLogin(InAsyncOp);
+		}
+		else
+		{
+			// TODO: EAS Logout
+
+			InAsyncOp.SetError(Errors::Unknown()); // TODO
+		}
+	})
+	.Enqueue(GetSerialQueue());
+
+	return Op->GetHandle();
+}
+
+TOnlineChainableAsyncOp<FAuthLogin, TSharedPtr<FEOSConnectLoginCredentials>> FAuthEOS::LoginEAS(TOnlineAsyncOp<FAuthLogin>& InAsyncOp)
+{
+	const FAuthLogin::Params& Params = InAsyncOp.GetParams();
+	auto FailureLambda = [](TOnlineAsyncOp<FAuthLogin>& Op) -> TSharedPtr<FEOSConnectLoginCredentials>
+	{
+		return nullptr;
+	};
+
 	EOS_Auth_LoginOptions LoginOptions = { };
 	LoginOptions.ApiVersion = EOS_AUTH_LOGIN_API_LATEST;
 	bool bContainsFlagsNone = false;
-	for (const FString& Scope : Op->GetParams().Scopes)
+	for (const FString& Scope : Params.Scopes)
 	{
 		EOS_EAuthScopeFlags ScopeFlag;
 		if (LexFromString(ScopeFlag, *Scope))
@@ -363,25 +466,25 @@ TOnlineAsyncOpHandle<FAuthLogin> FAuthEOS::Login(FAuthLogin::Params&& Params)
 		else
 		{
 			UE_LOG(LogTemp, Warning, TEXT("Invalid ScopeFlag=[%s]"), *Scope);
-			Op->SetError(Errors::Unknown());
-			return Op->GetHandle();
+			InAsyncOp.SetError(Errors::Unknown());
+			return InAsyncOp.Then(FailureLambda);
 		}
 	}
 	// TODO:  Where to put default scopes?
 	if (!bContainsFlagsNone && LoginOptions.ScopeFlags == EOS_EAuthScopeFlags::EOS_AS_NoFlags)
 	{
 		LoginOptions.ScopeFlags = EOS_EAuthScopeFlags::EOS_AS_BasicProfile | EOS_EAuthScopeFlags::EOS_AS_FriendsList | EOS_EAuthScopeFlags::EOS_AS_Presence;
-	}	
+	}
 
 	// First, we'll check if any subtypes have been specified
 	FString CredentialTypeStr;
 	FString ExternalCredentialTypeStr;
-	Op->GetParams().CredentialsType.Split(FString(TEXT(":")), &CredentialTypeStr, &ExternalCredentialTypeStr);
+	Params.CredentialsType.Split(FString(TEXT(":")), &CredentialTypeStr, &ExternalCredentialTypeStr);
 
 	// If no subtype was specified, we'll treat the string passed as the type
 	if (CredentialTypeStr.IsEmpty())
 	{
-		CredentialTypeStr = Op->GetParams().CredentialsType;
+		CredentialTypeStr = Params.CredentialsType;
 	}
 
 	FEOSAuthCredentials Credentials;
@@ -396,29 +499,29 @@ TOnlineAsyncOpHandle<FAuthLogin> FAuthEOS::Login(FAuthLogin::Params&& Params)
 				FAuthEOSConfig AuthEOSConfig;
 				LoadConfig(AuthEOSConfig);
 				ExternalCredentialTypeStr = AuthEOSConfig.DefaultExternalCredentialTypeStr;
-			}			
-			
+			}
+
 			EOS_EExternalCredentialType ExternalType;
 			if (LexFromString(ExternalType, *ExternalCredentialTypeStr))
 			{
 				Credentials.ExternalType = ExternalType;
 			}
 
-			Credentials.SetToken(Op->GetParams().CredentialsToken);
+			Credentials.SetToken(Params.CredentialsToken);
 			break;
 		case EOS_ELoginCredentialType::EOS_LCT_ExchangeCode:
 			// This is how the Epic launcher will pass credentials to you
 			Credentials.IdAnsi[0] = '\0';
-			FCStringAnsi::Strncpy(Credentials.TokenAnsi, TCHAR_TO_UTF8(*Op->GetParams().CredentialsToken.Get<FString>()), EOS_MAX_TOKEN_SIZE);
+			FCStringAnsi::Strncpy(Credentials.TokenAnsi, TCHAR_TO_UTF8(*Params.CredentialsToken.Get<FString>()), EOS_MAX_TOKEN_SIZE);
 			break;
 		case EOS_ELoginCredentialType::EOS_LCT_Password:
-			FCStringAnsi::Strncpy(Credentials.IdAnsi, TCHAR_TO_UTF8(*Op->GetParams().CredentialsId), EOS_OSS_STRING_BUFFER_LENGTH);
-			FCStringAnsi::Strncpy(Credentials.TokenAnsi, TCHAR_TO_UTF8(*Op->GetParams().CredentialsToken.Get<FString>()), EOS_MAX_TOKEN_SIZE);
+			FCStringAnsi::Strncpy(Credentials.IdAnsi, TCHAR_TO_UTF8(*Params.CredentialsId), EOS_OSS_STRING_BUFFER_LENGTH);
+			FCStringAnsi::Strncpy(Credentials.TokenAnsi, TCHAR_TO_UTF8(*Params.CredentialsToken.Get<FString>()), EOS_MAX_TOKEN_SIZE);
 			break;
 		case EOS_ELoginCredentialType::EOS_LCT_Developer:
 			// This is auth via the EOS auth tool
-			FCStringAnsi::Strncpy(Credentials.IdAnsi, TCHAR_TO_UTF8(*Op->GetParams().CredentialsId), EOS_OSS_STRING_BUFFER_LENGTH);
-			FCStringAnsi::Strncpy(Credentials.TokenAnsi, TCHAR_TO_UTF8(*Op->GetParams().CredentialsToken.Get<FString>()), EOS_MAX_TOKEN_SIZE);
+			FCStringAnsi::Strncpy(Credentials.IdAnsi, TCHAR_TO_UTF8(*Params.CredentialsId), EOS_OSS_STRING_BUFFER_LENGTH);
+			FCStringAnsi::Strncpy(Credentials.TokenAnsi, TCHAR_TO_UTF8(*Params.CredentialsToken.Get<FString>()), EOS_MAX_TOKEN_SIZE);
 			break;
 		case EOS_ELoginCredentialType::EOS_LCT_AccountPortal:
 			// This is auth via the EOS Account Portal
@@ -431,19 +534,19 @@ TOnlineAsyncOpHandle<FAuthLogin> FAuthEOS::Login(FAuthLogin::Params&& Params)
 			Credentials.Token = nullptr;
 			break;
 		default:
-			UE_LOG(LogTemp, Warning, TEXT("Unsupported CredentialsType=[%s]"), *Op->GetParams().CredentialsType);
-			Op->SetError(Errors::Unknown()); // TODO
-			return Op->GetHandle();
+			UE_LOG(LogTemp, Warning, TEXT("Unsupported CredentialsType=[%s]"), *Params.CredentialsType);
+			InAsyncOp.SetError(Errors::Unknown()); // TODO
+			return InAsyncOp.Then(FailureLambda);
 		}
 	}
 	else
 	{
-		UE_LOG(LogTemp, Warning, TEXT("Invalid CredentialsType=[%s]"), *Op->GetParams().CredentialsType);
-		Op->SetError(Errors::Unknown()); // TODO
-		return Op->GetHandle();
+		UE_LOG(LogTemp, Warning, TEXT("Invalid CredentialsType=[%s]"), *Params.CredentialsType);
+		InAsyncOp.SetError(Errors::Unknown()); // TODO
+		return InAsyncOp.Then(FailureLambda);
 	}
 
-	Op->Then([this, LoginOptions, Credentials](TOnlineAsyncOp<FAuthLogin>& InAsyncOp) mutable
+	return InAsyncOp.Then([this, LoginOptions, Credentials](TOnlineAsyncOp<FAuthLogin>& InAsyncOp) mutable
 	{
 		LoginOptions.Credentials = &Credentials;
 		return EOS_Async<EOS_Auth_LoginCallbackInfo>(EOS_Auth_Login, AuthHandle, LoginOptions);
@@ -480,7 +583,7 @@ TOnlineAsyncOpHandle<FAuthLogin> FAuthEOS::Login(FAuthLogin::Params&& Params)
 
 		return MakeFulfilledPromise<const EOS_Auth_LinkAccountCallbackInfo*>().GetFuture();
 	})
-	.Then([this](TOnlineAsyncOp<FAuthLogin>& InAsyncOp, const EOS_Auth_LinkAccountCallbackInfo* Data)
+	.Then([this](TOnlineAsyncOp<FAuthLogin>& InAsyncOp, const EOS_Auth_LinkAccountCallbackInfo* Data) -> TFuture<TSharedPtr<FEOSConnectLoginCredentials>>
 	{
 		UE_LOG(LogTemp, Verbose, TEXT("[FAuthEOS::Login] EOS_Auth_LinkAccount Result: [%s]"), (Data == nullptr) ? TEXT("Null") : *LexToString(Data->ResultCode));
 		
@@ -488,8 +591,7 @@ TOnlineAsyncOpHandle<FAuthLogin> FAuthEOS::Login(FAuthLogin::Params&& Params)
 		if (Data == nullptr || Data->ResultCode == EOS_EResult::EOS_Success)
 		{
 			EOS_EpicAccountId EpicAccountId = nullptr;
-
-			if (Data == NULL)
+			if (Data == nullptr)
 			{
 				EpicAccountId = *InAsyncOp.Data.Get<EOS_EpicAccountId>(TEXT("EpicAccountId"));
 			}
@@ -497,7 +599,6 @@ TOnlineAsyncOpHandle<FAuthLogin> FAuthEOS::Login(FAuthLogin::Params&& Params)
 			{
 				// We cache the Epic Account Id to use it in later stages of the login process
 				InAsyncOp.Data.Set(TEXT("EpicAccountId"), Data->LocalUserId);
-
 				EpicAccountId = Data->LocalUserId;
 			}
 
@@ -512,21 +613,12 @@ TOnlineAsyncOpHandle<FAuthLogin> FAuthEOS::Login(FAuthLogin::Params&& Params)
 
 			if (CopyResult == EOS_EResult::EOS_Success)
 			{
-				EOS_Connect_Credentials ConnectLoginCredentials = { };
-				ConnectLoginCredentials.ApiVersion = EOS_CONNECT_CREDENTIALS_API_LATEST;
-				ConnectLoginCredentials.Type = EOS_EExternalCredentialType::EOS_ECT_EPIC;
-				ConnectLoginCredentials.Token = AuthToken->AccessToken;
-
-				EOS_Connect_LoginOptions ConnectLoginOptions = { };
-				ConnectLoginOptions.ApiVersion = EOS_CONNECT_LOGIN_API_LATEST;
-				ConnectLoginOptions.Credentials = &ConnectLoginCredentials;
-
-				return EOS_Async<EOS_Connect_LoginCallbackInfo>(EOS_Connect_Login, ConnectHandle, ConnectLoginOptions);
+				TSharedPtr<FEOSConnectLoginCredentials> ConnectLoginCredentials = MakeShared<FEOSConnectLoginCredentialsEAS>(AuthToken);
+				return MakeFulfilledPromise<TSharedPtr<FEOSConnectLoginCredentials>>(MoveTemp(ConnectLoginCredentials)).GetFuture();
 			}
 			else
 			{
 				// TODO: EAS Logout
-
 				InAsyncOp.SetError(Errors::Unknown()); // TODO
 			}
 		}
@@ -535,64 +627,25 @@ TOnlineAsyncOpHandle<FAuthLogin> FAuthEOS::Login(FAuthLogin::Params&& Params)
 			InAsyncOp.SetError(Errors::Unknown()); // TODO
 		}
 
-		return MakeFulfilledPromise<const EOS_Connect_LoginCallbackInfo*>().GetFuture();
-	})
-	.Then([this](TOnlineAsyncOp<FAuthLogin>& InAsyncOp, const EOS_Connect_LoginCallbackInfo* Data)
-	{
-		UE_LOG(LogTemp, Verbose, TEXT("[FAuthEOS::Login] EOS_Connect_Login Result: [%s]"), *LexToString(Data->ResultCode));
-
-		if (Data->ResultCode == EOS_EResult::EOS_Success)
-		{
-			// We cache the Product User Id to use it in later stages of the login process
-			InAsyncOp.Data.Set(TEXT("ProductUserId"), Data->LocalUserId);
-
-			ProcessSuccessfulLogin(InAsyncOp);
-		}
-		else if (Data->ResultCode == EOS_EResult::EOS_InvalidUser && Data->ContinuanceToken != nullptr)
-		{
-			EOS_Connect_CreateUserOptions ConnectCreateUserOptions = { };
-			ConnectCreateUserOptions.ApiVersion = EOS_CONNECT_CREATEUSER_API_LATEST;
-			ConnectCreateUserOptions.ContinuanceToken = Data->ContinuanceToken;
-
-			return EOS_Async<EOS_Connect_CreateUserCallbackInfo>(EOS_Connect_CreateUser, ConnectHandle, ConnectCreateUserOptions);
-		}
-		else
-		{
-			// TODO: EAS Logout
-
-			InAsyncOp.SetError(Errors::Unknown()); // TODO
-		}
-
-		return MakeFulfilledPromise<const EOS_Connect_CreateUserCallbackInfo*>().GetFuture();
-	})
-	.Then([this](TOnlineAsyncOp<FAuthLogin>& InAsyncOp, const EOS_Connect_CreateUserCallbackInfo* Data)
-	{
-		UE_LOG(LogTemp, Verbose, TEXT("[FAuthEOS::Login] EOS_Connect_CreateUser Result: [%s]"), *LexToString(Data->ResultCode));
-
-		if (Data->ResultCode == EOS_EResult::EOS_Success)
-		{
-			// We cache the Product User Id to use it in later stages of the login process
-			InAsyncOp.Data.Set(TEXT("ProductUserId"), Data->LocalUserId);
-
-			ProcessSuccessfulLogin(InAsyncOp);
-		}
-		else
-		{
-			// TODO: EAS Logout
-
-			InAsyncOp.SetError(Errors::Unknown()); // TODO
-		}
-	})
-	.Enqueue(GetSerialQueue());
+		return MakeFulfilledPromise<TSharedPtr<FEOSConnectLoginCredentials>>(nullptr).GetFuture();
+	});
+}
 
-	return Op->GetHandle();
+TSharedPtr<FEOSConnectLoginCredentials> FAuthEOS::MakeConnectLoginCredentials(TOnlineAsyncOp<FAuthLogin>& InAsyncOp)
+{
+	const FAuthLogin::Params& LoginParams = InAsyncOp.GetParams();
+	TSharedRef<FEOSConnectLoginCredentials> Credentials = MakeShared<FEOSConnectLoginCredentialsEOS>();
+	FEOSConnectLoginCredentialsEOS& CredentialsEOS = static_cast<FEOSConnectLoginCredentialsEOS&>(*Credentials);
+	LexFromString(CredentialsEOS.Type, *LoginParams.CredentialsType);
+	CredentialsEOS.SetToken(LoginParams.CredentialsToken);
+	return Credentials;
 }
 
 void FAuthEOS::ProcessSuccessfulLogin(TOnlineAsyncOp<FAuthLogin>& InAsyncOp)
 {
-	const EOS_EpicAccountId EpicAccountId = *InAsyncOp.Data.Get<EOS_EpicAccountId>(TEXT("EpicAccountId"));
+	const EOS_EpicAccountId* EpicAccountId = InAsyncOp.Data.Get<EOS_EpicAccountId>(TEXT("EpicAccountId"));
 	const EOS_ProductUserId ProductUserId = *InAsyncOp.Data.Get<EOS_ProductUserId>(TEXT("ProductUserId"));
-	const FOnlineAccountIdHandle LocalUserId = CreateAccountId(EpicAccountId, ProductUserId);
+	const FOnlineAccountIdHandle LocalUserId = CreateAccountId(EpicAccountId ? *EpicAccountId : nullptr, ProductUserId);
 
 	UE_LOG(LogTemp, Verbose, TEXT("[FAuthEOS::Login] Successfully logged in as [%s]"), *ToLogString(LocalUserId));
 
@@ -601,21 +654,24 @@ void FAuthEOS::ProcessSuccessfulLogin(TOnlineAsyncOp<FAuthLogin>& InAsyncOp)
 	AccountInfo->UserId = LocalUserId;
 	AccountInfo->LoginStatus = ELoginStatus::LoggedIn;
 
-	// Get display name
-	if (EOS_HUserInfo UserInfoHandle = EOS_Platform_GetUserInfoInterface(static_cast<FOnlineServicesEOS&>(GetServices()).GetEOSPlatformHandle()))
+	if (EpicAccountId)
 	{
-		EOS_UserInfo_CopyUserInfoOptions Options = { };
-		Options.ApiVersion = EOS_USERINFO_COPYUSERINFO_API_LATEST;
-		Options.LocalUserId = EpicAccountId;
-		Options.TargetUserId = EpicAccountId;
+		// Get display name
+		if (EOS_HUserInfo UserInfoHandle = EOS_Platform_GetUserInfoInterface(static_cast<FOnlineServicesEOS&>(GetServices()).GetEOSPlatformHandle()))
+		{
+			EOS_UserInfo_CopyUserInfoOptions Options = { };
+			Options.ApiVersion = EOS_USERINFO_COPYUSERINFO_API_LATEST;
+			Options.LocalUserId = *EpicAccountId;
+			Options.TargetUserId = *EpicAccountId;
 
-		EOS_UserInfo* UserInfo = nullptr;
+			EOS_UserInfo* UserInfo = nullptr;
 
-		EOS_EResult CopyResult = EOS_UserInfo_CopyUserInfo(UserInfoHandle, &Options, &UserInfo);
-		if (CopyResult == EOS_EResult::EOS_Success)
-		{
-			AccountInfo->DisplayName = UTF8_TO_TCHAR(UserInfo->DisplayName);
-			EOS_UserInfo_Release(UserInfo);
+			EOS_EResult CopyResult = EOS_UserInfo_CopyUserInfo(UserInfoHandle, &Options, &UserInfo);
+			if (CopyResult == EOS_EResult::EOS_Success)
+			{
+				AccountInfo->DisplayName = UTF8_TO_TCHAR(UserInfo->DisplayName);
+				EOS_UserInfo_Release(UserInfo);
+			}
 		}
 	}
 
diff --git a/Engine/Plugins/Online/Experimental/OnlineServicesEOS/Source/Private/Online/OnlineServicesEOS.cpp b/Engine/Plugins/Online/Experimental/OnlineServicesEOS/Source/Private/Online/OnlineServicesEOS.cpp
--- a/Engine/Plugins/Online/Experimental/OnlineServicesEOS/Source/Private/Online/OnlineServicesEOS.cpp
+++ b/Engine/Plugins/Online/Experimental/OnlineServicesEOS/Source/Private/Online/OnlineServicesEOS.cpp
@@ -29,6 +29,7 @@ struct FEOSPlatformConfig
 	FString DeploymentId;
 	FString ClientId;
 	FString ClientSecret;
+	bool bUseEAS = false;
 };
 
 namespace Meta {
@@ -38,7 +39,8 @@ BEGIN_ONLINE_STRUCT_META(FEOSPlatformConfig)
 	ONLINE_STRUCT_FIELD(FEOSPlatformConfig, SandboxId),
 	ONLINE_STRUCT_FIELD(FEOSPlatformConfig, DeploymentId),
 	ONLINE_STRUCT_FIELD(FEOSPlatformConfig, ClientId),
-	ONLINE_STRUCT_FIELD(FEOSPlatformConfig, ClientSecret)
+	ONLINE_STRUCT_FIELD(FEOSPlatformConfig, ClientSecret),
+	ONLINE_STRUCT_FIELD(FEOSPlatformConfig, bUseEAS)
 END_ONLINE_STRUCT_META()
 
 /* Meta */ }
@@ -50,10 +52,16 @@ FOnlineServicesEOS::FOnlineServicesEOS()
 
 void FOnlineServicesEOS::RegisterComponents()
 {
-	Components.Register<FAuthEOS>(*this);
-	Components.Register<FFriendsEOS>(*this);
+	FEOSPlatformConfig EOSPlatformConfig;
+	LoadConfig(EOSPlatformConfig);
+
+	Components.Register<FAuthEOS>(*this, EOSPlatformConfig.bUseEAS);
+	if (EOSPlatformConfig.bUseEAS)
+	{
+		Components.Register<FFriendsEOS>(*this);
+		Components.Register<FPresenceEOS>(*this);
+	}
 	Components.Register<FLobbiesEOS>(*this);
-	Components.Register<FPresenceEOS>(*this);
 	Components.Register<FExternalUIEOS>(*this);
 	FOnlineServicesCommon::RegisterComponents();
 }
diff --git a/Engine/Plugins/Online/Experimental/OnlineServicesEOS/Source/Public/Online/AuthEOS.h b/Engine/Plugins/Online/Experimental/OnlineServicesEOS/Source/Public/Online/AuthEOS.h
--- a/Engine/Plugins/Online/Experimental/OnlineServicesEOS/Source/Public/Online/AuthEOS.h
+++ b/Engine/Plugins/Online/Experimental/OnlineServicesEOS/Source/Public/Online/AuthEOS.h
@@ -20,7 +20,7 @@ class ONLINESERVICESEOS_API FAuthEOS : public FAuthCommon
 public:
 	using Super = FAuthCommon;
 
-	FAuthEOS(FOnlineServicesEOS& InOwningSubsystem);
+	FAuthEOS(FOnlineServicesEOS& InOwningSubsystem, bool bInUseEAS);
 	virtual void Initialize() override;
 	virtual void PreShutdown() override;
 	virtual TOnlineAsyncOpHandle<FAuthLogin> Login(FAuthLogin::Params&& Params) override;
@@ -42,6 +42,9 @@ public:
 	TFunction<TFuture<TArray<FOnlineAccountIdHandle>>(FOnlineAsyncOp& InAsyncOp, const TArray<EOS_ProductUserId>& ProductUserIds)> ResolveProductIdsFn();
 
 protected:
+	TOnlineChainableAsyncOp<FAuthLogin, TSharedPtr<class FEOSConnectLoginCredentials>> LoginEAS(TOnlineAsyncOp<FAuthLogin>& InAsyncOp);
+	TSharedPtr<class FEOSConnectLoginCredentials> MakeConnectLoginCredentials(TOnlineAsyncOp<FAuthLogin>& InAsyncOp);
+
 	void OnEOSLoginStatusChanged(FOnlineAccountIdHandle LocalUserId, ELoginStatus PreviousStatus, ELoginStatus CurrentStatus);
 	TResult<FOnlineAccountIdHandle, FOnlineError> GetAccountIdByPlatformUserId(FPlatformUserId PlatformUserId) const;
 
@@ -58,6 +61,9 @@ protected:
 	EOS_HAuth AuthHandle;
 	EOS_HConnect ConnectHandle;
 	EOS_NotificationId NotifyLoginStatusChangedNotificationId = 0;
+
+	/** Are we configured to use EAS? */
+	bool bUseEAS = false;
 };
 
 /* UE::Online */ }
