Refactor of Auth Token operations to use EOS_MAX_TOKEN_SIZE.

diff --git a/Engine/Plugins/Online/Experimental/OnlineServicesEOS/Source/Private/Online/AuthEOS.cpp b/Engine/Plugins/Online/Experimental/OnlineServicesEOS/Source/Private/Online/AuthEOS.cpp
--- a/Engine/Plugins/Online/Experimental/OnlineServicesEOS/Source/Private/Online/AuthEOS.cpp
+++ b/Engine/Plugins/Online/Experimental/OnlineServicesEOS/Source/Private/Online/AuthEOS.cpp
@@ -94,7 +94,7 @@ struct FEOSAuthCredentials :
 		Id = IdAnsi;
 		Token = TokenAnsi;
 
-		uint32_t InOutBufferLength = EOS_OSS_STRING_BUFFER_LENGTH;
+		uint32_t InOutBufferLength = EOS_MAX_TOKEN_SIZE;
 		EOS_ByteArray_ToString(InToken.GetData(), InToken.Num(), TokenAnsi, &InOutBufferLength);
 	}
 
@@ -116,7 +116,7 @@ struct FEOSAuthCredentials :
 		if (InToken.IsType<TArray<uint8>>())
 		{
 			const TArray<uint8>& TokenData = InToken.Get<TArray<uint8>>();
-			uint32_t InOutBufferLength = EOS_OSS_STRING_BUFFER_LENGTH;
+			uint32_t InOutBufferLength = EOS_MAX_TOKEN_SIZE;
 			EOS_ByteArray_ToString(TokenData.GetData(), TokenData.Num(), TokenAnsi, &InOutBufferLength);
 		}
 		else if (InToken.IsType<FString>())
