Check if chainable op was set to failure before calling Enqueue

diff --git a/Engine/Plugins/Online/Experimental/OnlineServicesEOS/Source/Private/Online/AuthEOS.cpp b/Engine/Plugins/Online/Experimental/OnlineServicesEOS/Source/Private/Online/AuthEOS.cpp
--- a/Engine/Plugins/Online/Experimental/OnlineServicesEOS/Source/Private/Online/AuthEOS.cpp
+++ b/Engine/Plugins/Online/Experimental/OnlineServicesEOS/Source/Private/Online/AuthEOS.cpp
@@ -383,6 +383,11 @@ TOnlineAsyncOpHandle<FAuthLogin> FAuthEOS::Login(FAuthLogin::Params&& Params)
 			return MakeConnectLoginCredentials(InAsyncOp);
 		});
 	}();
+	// Check if the above steps completed (failed) the operation
+	if (ConnectLoginOp.GetOwningOperation().IsComplete())
+	{
+		return ConnectLoginOp.GetOwningOperation().GetHandle();
+	}
 	ConnectLoginOp.Then([this](TOnlineAsyncOp<FAuthLogin>& InAsyncOp, TSharedPtr<FEOSConnectLoginCredentials>&& InConnectLoginCredentials)
 	{
 		EOS_Connect_LoginOptions ConnectLoginOptions = { };
