Wrap socket subsystem reference with WITH_ENGINE so OnlineServicesEOS can compile without Engine

diff --git a/Engine/Plugins/Online/Experimental/OnlineServicesEOS/Source/Private/Online/SocketSubsystemEOSUtils_OnlineServicesEOS.cpp b/Engine/Plugins/Online/Experimental/OnlineServicesEOS/Source/Private/Online/SocketSubsystemEOSUtils_OnlineServicesEOS.cpp
--- a/Engine/Plugins/Online/Experimental/OnlineServicesEOS/Source/Private/Online/SocketSubsystemEOSUtils_OnlineServicesEOS.cpp
+++ b/Engine/Plugins/Online/Experimental/OnlineServicesEOS/Source/Private/Online/SocketSubsystemEOSUtils_OnlineServicesEOS.cpp
@@ -1,5 +1,6 @@
 // Copyright Epic Games, Inc. All Rights Reserved.
 
+#if WITH_ENGINE
 #include "SocketSubsystemEOSUtils_OnlineServicesEOS.h"
 
 #include "Online/OnlineServicesEOS.h"
@@ -87,4 +88,6 @@ FString FSocketSubsystemEOSUtils_OnlineServicesEOS::GetSessionId()
 	return Result;
 }
 
-/* UE::Online */}
\ No newline at end of file
+/* UE::Online */}
+
+#endif // WITH_ENGINE
diff --git a/Engine/Plugins/Online/Experimental/OnlineServicesEOS/Source/Private/Online/SocketSubsystemEOSUtils_OnlineServicesEOS.h b/Engine/Plugins/Online/Experimental/OnlineServicesEOS/Source/Private/Online/SocketSubsystemEOSUtils_OnlineServicesEOS.h
--- a/Engine/Plugins/Online/Experimental/OnlineServicesEOS/Source/Private/Online/SocketSubsystemEOSUtils_OnlineServicesEOS.h
+++ b/Engine/Plugins/Online/Experimental/OnlineServicesEOS/Source/Private/Online/SocketSubsystemEOSUtils_OnlineServicesEOS.h
@@ -2,6 +2,8 @@
 
 #pragma once
 
+#if WITH_ENGINE
+
 #include "SocketSubsystemEOS.h"
 
 namespace UE::Online {
@@ -21,4 +23,6 @@ private:
 	FOnlineServicesEOS& ServicesEOS;
 };
 
-/* UE::Online */}
\ No newline at end of file
+/* UE::Online */}
+
+#endif // WITH_ENGINE
\ No newline at end of file
