- Fix passing FString, not the underlying TCHAR*, to FTCHARToUTF8

diff --git a/Engine/Plugins/Online/Experimental/OnlineServicesEOS/Source/Private/Online/OnlineServicesEOS.cpp b/Engine/Plugins/Online/Experimental/OnlineServicesEOS/Source/Private/Online/OnlineServicesEOS.cpp
--- a/Engine/Plugins/Online/Experimental/OnlineServicesEOS/Source/Private/Online/OnlineServicesEOS.cpp
+++ b/Engine/Plugins/Online/Experimental/OnlineServicesEOS/Source/Private/Online/OnlineServicesEOS.cpp
@@ -88,7 +88,7 @@ void FOnlineServicesEOS::Initialize()
 	const FTCHARToUTF8 DeploymentId(*EOSPlatformConfig.DeploymentId);
 	const FTCHARToUTF8 ClientId(*EOSPlatformConfig.ClientId);
 	const FTCHARToUTF8 ClientSecret(*EOSPlatformConfig.ClientSecret);
-	const FTCHARToUTF8 CacheDirectory(SDKManager->GetCacheDirBase() / TEXT("OnlineServicesEOS"));
+	const FTCHARToUTF8 CacheDirectory(*(SDKManager->GetCacheDirBase() / TEXT("OnlineServicesEOS")));
 
 	EOS_Platform_Options PlatformOptions = {};
 	PlatformOptions.ApiVersion = EOS_PLATFORM_OPTIONS_API_LATEST;
