- Fix ResolveAccountIds assuming there will be an EAS id

diff --git a/Engine/Plugins/Online/Experimental/OnlineServicesEOS/Source/Private/Online/AuthEOS.cpp b/Engine/Plugins/Online/Experimental/OnlineServicesEOS/Source/Private/Online/AuthEOS.cpp
--- a/Engine/Plugins/Online/Experimental/OnlineServicesEOS/Source/Private/Online/AuthEOS.cpp
+++ b/Engine/Plugins/Online/Experimental/OnlineServicesEOS/Source/Private/Online/AuthEOS.cpp
@@ -1025,9 +1025,13 @@ TFuture<TArray<FOnlineAccountIdHandle>> FAuthEOS::ResolveAccountIds(const FOnlin
 					Options.TargetProductUserId = ProductUserId;
 					FEpicAccountIdStrBuffer EpicAccountIdStr;
 					int32_t BufferLength = sizeof(EpicAccountIdStr);
-					verify(EOS_Connect_GetProductUserIdMapping(ConnectHandle, &Options, EpicAccountIdStr, &BufferLength) == EOS_EResult::EOS_Success);
-					const EOS_EpicAccountId EpicAccountId = EOS_EpicAccountId_FromString(EpicAccountIdStr);
-					check(EOS_EpicAccountId_IsValid(EpicAccountId));
+					EOS_EpicAccountId EpicAccountId = nullptr;
+					const EOS_EResult Result = EOS_Connect_GetProductUserIdMapping(ConnectHandle, &Options, EpicAccountIdStr, &BufferLength);
+					if (Result == EOS_EResult::EOS_Success)
+					{
+						EpicAccountId = EOS_EpicAccountId_FromString(EpicAccountIdStr);
+						check(EOS_EpicAccountId_IsValid(EpicAccountId));
+					}
 					AccountId = CreateAccountId(EpicAccountId, ProductUserId);
 				}
 				AccountIds.Emplace(MoveTemp(AccountId));
