Add basic support for converting EOS_EResult into an FOnlineError.

Also added some common definitions that we were missing based on analysis from the eos error codes

diff --git a/Engine/Plugins/Online/Experimental/OnlineServices/Source/OnlineServicesInterface/Public/Online/OnlineErrorDefinitions.h b/Engine/Plugins/Online/Experimental/OnlineServices/Source/OnlineServicesInterface/Public/Online/OnlineErrorDefinitions.h
--- a/Engine/Plugins/Online/Experimental/OnlineServices/Source/OnlineServicesInterface/Public/Online/OnlineErrorDefinitions.h
+++ b/Engine/Plugins/Online/Experimental/OnlineServices/Source/OnlineServicesInterface/Public/Online/OnlineErrorDefinitions.h
@@ -8,6 +8,7 @@ namespace UE::Online::Errors {
 
 UE_ONLINE_ERROR_CATEGORY(Common, Engine, 0x1, "Online Services")
 
+UE_ONLINE_ERROR_COMMON(Common, Success, 0, TEXT("success"), LOCTEXT("Success", "Success"))
 UE_ONLINE_ERROR_COMMON(Common, NoConnection, 1, TEXT("no_connection"), LOCTEXT("NotConnected", "No valid connection"))
 UE_ONLINE_ERROR_COMMON(Common, RequestFailure, 2, TEXT("request_failure"), LOCTEXT("RequestFailure", "Failed to send request"))
 UE_ONLINE_ERROR_COMMON(Common, InvalidCreds, 3, TEXT("invalid_creds"), LOCTEXT("InvalidCreds", "Invalid credentials"))
@@ -23,8 +24,12 @@ UE_ONLINE_ERROR_COMMON(Common, IncompatibleVersion, 12, TEXT("incompatible_versi
 UE_ONLINE_ERROR_COMMON(Common, NotConfigured, 13, TEXT("not_configured"), LOCTEXT("NotConfigured", "No valid configuration"))
 UE_ONLINE_ERROR_COMMON(Common, NotImplemented, 14, TEXT("not_implemented"), LOCTEXT("NotImplemented", "Not implemented"))
 UE_ONLINE_ERROR_COMMON(Common, MissingInterface, 15, TEXT("missing_interface"), LOCTEXT("MissingInterface", "Interface not found"))
-UE_ONLINE_ERROR_COMMON(Common, Cancelled, 16, TEXT("cancelled"), LOCTEXT("Canceled", "Operation was cancelled"))
-UE_ONLINE_ERROR_COMMON(Common, Unknown, 17, TEXT("unknown"), LOCTEXT("Unknown", "Unknown Error"))
+UE_ONLINE_ERROR_COMMON(Common, Cancelled, 16, TEXT("cancelled"), LOCTEXT("Cancelled", "Operation was cancelled"))
+UE_ONLINE_ERROR_COMMON(Common, NotLoggedIn, 17, TEXT("not_logged_in"), LOCTEXT("NotLoggedIn", "User is not logged in"))
+UE_ONLINE_ERROR_COMMON(Common, NotFound, 18, TEXT("not_found"), LOCTEXT("NotFound", "Request not found"))
+UE_ONLINE_ERROR_COMMON(Common, WillRetry, 19, TEXT("will_retry"), LOCTEXT("WillRetry", "Retrying request"))
+UE_ONLINE_ERROR_COMMON(Common, Timeout, 20, TEXT("timeout"), LOCTEXT("Timeout", "Operation timed out"))
+UE_ONLINE_ERROR_COMMON(Common, Unknown, 21, TEXT("unknown"), LOCTEXT("Unknown", "Unknown Error"))
 
 } /* namespace UE::Online::Errors */
 
diff --git a/Engine/Plugins/Online/Experimental/OnlineServicesEOS/Source/Public/Online/ErrorsEOS.h b/Engine/Plugins/Online/Experimental/OnlineServicesEOS/Source/Public/Online/ErrorsEOS.h
new file mode 100644
--- /dev/null
+++ b/Engine/Plugins/Online/Experimental/OnlineServicesEOS/Source/Public/Online/ErrorsEOS.h
@@ -0,0 +1,89 @@
+// Copyright Epic Games, Inc. All Rights Reserved.
+#pragma once
+
+#include "Online/OnlineError.h"
+#include "Online/OnlineErrorCode.h"
+#include "Online/PresenceEOS.h"
+
+/*
+*	Proper usage:
+		Errors::FromEOSResult(EOSResult::EOS_PlayerDataStorage_FileSizeTooLarge);
+	Certain EOS errors are predefined to have a common error parent type (see Internal_EOSWrapInner). i.e.
+		(Errors::FromEOSResult(EOSResult::EOS_NoConnection) == Errors::NoConnection()) == true
+*/
+
+namespace UE::Online::Errors {
+
+	UE_ONLINE_ERROR_CATEGORY(EOS, Engine, 0x4, "EOS")
+
+	/** This doesn't wrap every single error in EOS, only the common ones that are strongly related to common errors */
+	inline FOnlineError Internal_EOSWrapInner(FOnlineError Error, EOS_EResult Result)
+	{
+		switch(Result)
+		{
+			case EOS_EResult::EOS_Success:				return Errors::Success(Error);
+			case EOS_EResult::EOS_NoConnection:			return Errors::NoConnection(Error);
+			case EOS_EResult::EOS_InvalidCredentials:	return Errors::InvalidCreds(Error);
+			case EOS_EResult::EOS_InvalidUser:			return Errors::InvalidUser(Error);
+			case EOS_EResult::EOS_InvalidAuth:			return Errors::InvalidAuth(Error);
+			case EOS_EResult::EOS_AccessDenied:			return Errors::AccessDenied(Error);
+			case EOS_EResult::EOS_MissingPermissions:	return Errors::AccessDenied(Error);
+			case EOS_EResult::EOS_TooManyRequests:		return Errors::TooManyRequests(Error);
+			case EOS_EResult::EOS_AlreadyPending:		return Errors::AlreadyPending(Error);
+			case EOS_EResult::EOS_InvalidParameters:	return Errors::InvalidParams(Error);
+			case EOS_EResult::EOS_InvalidRequest:		return Errors::InvalidParams(Error);
+			case EOS_EResult::EOS_UnrecognizedResponse: return Errors::InvalidResults(Error);
+			case EOS_EResult::EOS_IncompatibleVersion:	return Errors::IncompatibleVersion(Error);
+			case EOS_EResult::EOS_NotConfigured:		return Errors::NotConfigured(Error);
+			case EOS_EResult::EOS_NotImplemented:		return Errors::NotImplemented(Error);
+			case EOS_EResult::EOS_Canceled:				return Errors::Cancelled(Error);
+			case EOS_EResult::EOS_NotFound:				return Errors::NotFound(Error);
+			case EOS_EResult::EOS_OperationWillRetry:	return Errors::WillRetry(Error);
+			case EOS_EResult::EOS_VersionMismatch:		return Errors::IncompatibleVersion(Error);
+			case EOS_EResult::EOS_LimitExceeded:		return Errors::TooManyRequests(Error);
+			case EOS_EResult::EOS_TimedOut:				return Errors::Timeout(Error);
+		}
+		return Error;
+	}
+
+	inline ErrorCodeType ErrorValueFromEOSResult(EOS_EResult Result)
+	{
+		return (ErrorCodeType)((uint32)Result);
+	}
+
+	inline ErrorCodeType ErrorCodeFromEOSResult(EOS_EResult Result)
+	{
+		return ErrorCode::Create(ErrorCode::System::ThirdPartyPlugin, ErrorCode::Category::EOS, ErrorValueFromEOSResult(Result));
+	}
+
+
+	inline FOnlineError FromEOSResult(EOS_EResult Result)
+	{
+		// these get moved around so make sure they have their own memories
+		FString ErrorStr = UTF8_TO_TCHAR(EOS_EResult_ToString(Result));
+		FText ErrorTxt = FText::FromString(UTF8_TO_TCHAR(EOS_EResult_ToString(Result)));
+		return Internal_EOSWrapInner(FOnlineError(ErrorCodeFromEOSResult(Result), MakeShared<FOnlineErrorDetails, ESPMode::ThreadSafe>(MoveTemp(ErrorStr), MoveTemp(ErrorTxt))), Result);
+	}
+
+} // UE::Online::Errors
+
+// These are extern'd out of the namespace for some of the Catch API that can't see these operators if they are inside
+ONLINESERVICESEOS_API inline bool operator==(const UE::Online::FOnlineError& Left, EOS_EResult Right)
+{
+	return Left == UE::Online::Errors::ErrorCodeFromEOSResult(Right);
+}
+
+ONLINESERVICESEOS_API inline bool operator==(EOS_EResult Left, const UE::Online::FOnlineError& Right)
+{
+	return Right == UE::Online::Errors::ErrorCodeFromEOSResult(Left);
+}
+
+ONLINESERVICESEOS_API inline bool operator!=(const UE::Online::FOnlineError& Left, EOS_EResult Right)
+{
+	return !(Left == UE::Online::Errors::ErrorCodeFromEOSResult(Right));
+}
+
+ONLINESERVICESEOS_API inline bool operator!=(EOS_EResult Left, const UE::Online::FOnlineError& Right)
+{
+	return !(Right == UE::Online::Errors::ErrorCodeFromEOSResult(Left));
+}
\ No newline at end of file
