Presence adapation for the new polished RFC. Implements changes in all of the header file, common implementation, and EOS implementation.

diff --git a/Engine/Plugins/Online/Experimental/OnlineServices/Source/OnlineServicesCommon/Private/Online/PresenceCommon.cpp b/Engine/Plugins/Online/Experimental/OnlineServices/Source/OnlineServicesCommon/Private/Online/PresenceCommon.cpp
--- a/Engine/Plugins/Online/Experimental/OnlineServices/Source/OnlineServicesCommon/Private/Online/PresenceCommon.cpp
+++ b/Engine/Plugins/Online/Experimental/OnlineServices/Source/OnlineServicesCommon/Private/Online/PresenceCommon.cpp
@@ -18,10 +18,96 @@ TOnlineAsyncOpHandle<FQueryPresence> FPresenceCommon::QueryPresence(FQueryPresen
 	Operation->SetError(Errors::NotImplemented());
 	return Operation->GetHandle();
 }
+/*
+      The common implementation of BatchQueryPresence sends out a bunch of asynchronous QueryPresence calls and combines them. 
+*/ 
+struct FBatchQueryPresenceHelperResult
+{
+public:
+	FBatchQueryPresence::Result Result;
+	FOnlineError ErrorState;
+};
+
+class FBatchQueryPresenceHelper
+{
+public:
+	FBatchQueryPresenceHelper(TOnlineAsyncOpRef<FBatchQueryPresence> InOp, TSharedPtr<IPresence> InPresence)
+		: Op(InOp)
+		, Presence(InPresence)
+		, ErrorState(Errors::Success())
+	{
+		NumPresencesRemaining = Op->GetParams().TargetUserIds.Num();
+	}
+
+	TFuture<FBatchQueryPresenceHelperResult> GetPromise()
+	{
+		for(FOnlineAccountIdHandle TargetUserId : Op->GetParams().TargetUserIds)
+		{
+			FQueryPresence::Params Params;
+			Params.LocalUserId = Op->GetParams().LocalUserId;
+			Params.TargetUserId = TargetUserId;
+			Params.bListenToChanges = Op->GetParams().bListenToChanges;
+
+			Presence->QueryPresence(MoveTemp(Params)).OnComplete([this](const TOnlineResult<FQueryPresence> NewOp) mutable
+			{
+				if (NewOp.IsOk())
+				{
+					// push results
+					Result.Presences.Add(NewOp.GetOkValue().Presence);
+
+					// final calculations
+					if (--NumPresencesRemaining < 0)
+					{
+						Promise.SetValue(FBatchQueryPresenceHelperResult{MoveTemp(Result), MoveTemp(ErrorState)});
+					}
+				}
+				else
+				{
+					// error
+					// we're only reporting the most recent one, more verbose logging on each individual error can be found in the logs
+					ErrorState = NewOp.GetErrorValue();
+				}
+			});
+		}
 
-TOnlineResult<FGetPresence> FPresenceCommon::GetPresence(FGetPresence::Params&& Params)
+		return Promise.GetFuture();
+	}
+
+	FBatchQueryPresence::Result Result;
+	TOnlineAsyncOpRef<FBatchQueryPresence> Op;
+	TSharedPtr<IPresence> Presence;
+	FOnlineError ErrorState;
+	int NumPresencesRemaining = 0;
+	TPromise<FBatchQueryPresenceHelperResult> Promise;
+};
+
+// Base implementation for the batch query function makes multiple async calls to query and returns them as a single output
+TOnlineAsyncOpHandle<FBatchQueryPresence> FPresenceCommon::BatchQueryPresence(FBatchQueryPresence::Params&& Params)
+{
+	TOnlineAsyncOpRef<FBatchQueryPresence> Op = GetOp<FBatchQueryPresence>(MoveTemp(Params));
+	TSharedRef<FBatchQueryPresenceHelper> Helper = MakeShared<FBatchQueryPresenceHelper>(Op, this->AsShared());
+	Op->Then([this, Helper](TOnlineAsyncOp<FBatchQueryPresence>& InAsyncOp) mutable
+	{
+		return Helper->GetPromise();
+	})
+	
+	.Then([this, Helper](TOnlineAsyncOp<FBatchQueryPresence>& InAsyncOp, FBatchQueryPresenceHelperResult Result) mutable
+	{
+		if (Result.ErrorState == Errors::Success())
+		{
+			InAsyncOp.SetResult(MoveTemp(Result.Result));
+		}
+		else 
+		{
+			InAsyncOp.SetError(MoveTemp(Result.ErrorState));
+		}
+	});
+	return Op->GetHandle();
+}
+
+TOnlineResult<FGetCachedPresence> FPresenceCommon::GetCachedPresence(FGetCachedPresence::Params&& Params)
 {
-	return TOnlineResult<FGetPresence>(Errors::NotImplemented());
+	return TOnlineResult<FGetCachedPresence>(Errors::NotImplemented());
 }
 
 TOnlineAsyncOpHandle<FUpdatePresence> FPresenceCommon::UpdatePresence(FUpdatePresence::Params&& Params)
@@ -31,6 +117,63 @@ TOnlineAsyncOpHandle<FUpdatePresence> FPresenceCommon::UpdatePresence(FUpdatePre
 	return Operation->GetHandle();
 }
 
+TOnlineAsyncOpHandle<FPartialUpdatePresence> FPresenceCommon::PartialUpdatePresence(FPartialUpdatePresence::Params&& Params)
+{
+	TOnlineResult<FGetCachedPresence> CachedPresence = GetCachedPresence({Params.LocalUserId, Params.LocalUserId});
+	if (CachedPresence.IsOk())
+	{
+		TSharedRef<FUserPresence> NewPresence = MakeShared<FUserPresence>();
+		TSharedRef<const FUserPresence> ExistingPresence = CachedPresence.GetOkValue().Presence;
+		*NewPresence = *ExistingPresence;
+
+		if (Params.Mutations.Status.IsSet())
+		{
+			NewPresence->Status = Params.Mutations.Status.GetValue();
+		}
+
+		if (Params.Mutations.GameStatus.IsSet())
+		{
+			NewPresence->GameStatus = Params.Mutations.GameStatus.GetValue();
+		}
+
+		if (Params.Mutations.Joinability.IsSet())
+		{
+			NewPresence->Joinability = Params.Mutations.Joinability.GetValue();
+		}
+
+		for (const TPair<FString, FPresenceVariant>& UpdatedProperty : Params.Mutations.UpdatedProperties)
+		{
+			NewPresence->Properties.Add(UpdatedProperty.Key, UpdatedProperty.Value);
+		}
+
+		for (const FString& RemovedProperty : Params.Mutations.RemovedProperties)
+		{
+			NewPresence->Properties.Remove(RemovedProperty);
+		}
+
+		TOnlineAsyncOpRef<FPartialUpdatePresence> Op = GetOp<FPartialUpdatePresence>(MoveTemp(Params));
+		TOnlineAsyncOpHandle<FUpdatePresence> UpdatePresenceResult = UpdatePresence({Op->GetParams().LocalUserId, NewPresence});
+		UpdatePresenceResult.OnComplete([this, Op](const TOnlineResult<FUpdatePresence> Result)
+		{
+			if (Result.IsOk())
+			{
+				Op->SetResult(FPartialUpdatePresence::Result());
+			}
+			else
+			{
+				FOnlineError Error = Result.GetErrorValue();
+				Op->SetError(MoveTemp(Error));
+			}
+		});
+
+		return Op->GetHandle();
+	}
+
+	TOnlineAsyncOpRef<FPartialUpdatePresence> Operation = GetOp<FPartialUpdatePresence>(MoveTemp(Params));
+	Operation->SetError(Errors::NotImplemented());
+	return Operation->GetHandle();
+}
+
 TOnlineEvent<void(const FPresenceUpdated&)> FPresenceCommon::OnPresenceUpdated()
 {
 	return OnPresenceUpdatedEvent;
diff --git a/Engine/Plugins/Online/Experimental/OnlineServices/Source/OnlineServicesCommon/Public/Online/PresenceCommon.h b/Engine/Plugins/Online/Experimental/OnlineServices/Source/OnlineServicesCommon/Public/Online/PresenceCommon.h
--- a/Engine/Plugins/Online/Experimental/OnlineServices/Source/OnlineServicesCommon/Public/Online/PresenceCommon.h
+++ b/Engine/Plugins/Online/Experimental/OnlineServices/Source/OnlineServicesCommon/Public/Online/PresenceCommon.h
@@ -18,8 +18,10 @@ public:
 
 	// IPresence
 	virtual TOnlineAsyncOpHandle<FQueryPresence> QueryPresence(FQueryPresence::Params&& Params) override;
-	virtual TOnlineResult<FGetPresence> GetPresence(FGetPresence::Params&& Params) override;
+	virtual TOnlineAsyncOpHandle<FBatchQueryPresence> BatchQueryPresence(FBatchQueryPresence::Params&& Params) override;
+	virtual TOnlineResult<FGetCachedPresence> GetCachedPresence(FGetCachedPresence::Params&& Params) override;
 	virtual TOnlineAsyncOpHandle<FUpdatePresence> UpdatePresence(FUpdatePresence::Params&& Params) override;
+	virtual TOnlineAsyncOpHandle<FPartialUpdatePresence> PartialUpdatePresence(FPartialUpdatePresence::Params&& Params) override;
 	virtual TOnlineEvent<void(const FPresenceUpdated&)> OnPresenceUpdated() override;
 
 protected:
diff --git a/Engine/Plugins/Online/Experimental/OnlineServices/Source/OnlineServicesInterface/Public/Online/Presence.h b/Engine/Plugins/Online/Experimental/OnlineServices/Source/OnlineServicesInterface/Public/Online/Presence.h
--- a/Engine/Plugins/Online/Experimental/OnlineServices/Source/OnlineServicesInterface/Public/Online/Presence.h
+++ b/Engine/Plugins/Online/Experimental/OnlineServices/Source/OnlineServicesInterface/Public/Online/Presence.h
@@ -5,12 +5,29 @@
 #include "Online/CoreOnline.h"
 #include "Online/OnlineAsyncOpHandle.h"
 #include "Online/OnlineMeta.h"
+#include "Online/OnlineError.h"
+
+#define LOCTEXT_NAMESPACE "Presence"
 
 namespace UE::Online {
 
 using FPresenceVariant = FString; // TODO:  Import FVariantData
 
-enum class EPresenceState : uint8
+namespace Errors {
+UE_ONLINE_ERROR_CATEGORY(Presence, Engine, 0x3, "Presence")
+//UE_ONLINE_ERROR(Presence, OnlyFriendsSupported, 0, TEXT("OnlyFriendsSupported"), LOCTEXT("OnlyFriendsSupported", "This platform only supports listening to presence updates for friends."))
+UE_ONLINE_ERROR_COMMON(Presence, CannotQueryLocalUsers, 1, TEXT("CannotQueryLocalUsers"), LOCTEXT("CannotQueryLocalUsers", "QueryPresence cannot be called on users logged into the game currently"))
+}
+
+// todo: add more potential errors here
+/* Common errors used:
+	NotFound
+	ServiceUnavailable
+	Forbidden
+	RateLimit
+*/
+
+enum class EUserPresenceStatus : uint8
 {
 	/** User is offline */
 	Offline,
@@ -18,10 +35,36 @@ enum class EPresenceState : uint8
 	Online,
 	/** User is away */
 	Away,
-	/** User is extended away */
+	/** User is away for >2 hours (can change depending on platform) */
 	ExtendedAway,
 	/** User is in do not disturb mode */
-	DoNotDisturb
+	DoNotDisturb,
+	/** Default */
+	Unknown
+};
+
+enum class EUserPresenceJoinability : uint8
+{
+	/** Anyone can matchmake/discover and join this session */
+	Public,
+	/**  Anyone trying to join must be a friend of a lobby member */
+	FriendsOnly,
+	/** Anyone trying to join must have been invited first */
+	InviteOnly,
+	/** Not currently accepting invites */
+	Private,
+	/** Default */
+	Unknown
+};
+
+enum class EUserPresenceGameStatus : uint8
+{
+	/** The user is playing the same game as you */
+	PlayingThisGame,
+	/** The user is playing a different game than you */
+	PlayingOtherGame,
+	/** Default */
+	Unknown
 };
 
 typedef TMap<FString, FPresenceVariant> FPresenceProperties;
@@ -31,10 +74,16 @@ struct FUserPresence
 	/** User whose presence this is */
 	FOnlineAccountIdHandle UserId;
 	/** Presence state */
-	EPresenceState State;
-	/** Presence string */
+	EUserPresenceStatus Status = EUserPresenceStatus::Unknown;
+	/** Session state */
+	EUserPresenceJoinability Joinability = EUserPresenceJoinability::Unknown;
+	/** Playing this game */
+	EUserPresenceGameStatus GameStatus = EUserPresenceGameStatus::Unknown;
+	/** Stringified representation of the user status (online, offline, whatever the platform wants to define) */
 	FString StatusString;
-	/** Presence properties */
+	/** Game-defined representation of the current game state (e.g. "Squads- 4 Remaining" or "Level 5 Warrior") */
+	FString RichPresenceString;
+	/** Session keys */
 	FPresenceProperties Properties;
 };
 	
@@ -48,6 +97,8 @@ struct FQueryPresence
 		FOnlineAccountIdHandle LocalUserId;
 		/** User to query the presence for */
 		FOnlineAccountIdHandle TargetUserId;
+		/** If true, then future presence updates for this user will be propagated via the OnPresenceUpdated event */
+		bool bListenToChanges = true;
 	};
 
 	struct Result
@@ -57,9 +108,31 @@ struct FQueryPresence
 	};
 };
 
-struct FGetPresence
+
+struct FBatchQueryPresence
 {
-	static constexpr TCHAR Name[] = TEXT("GetPresence");
+	static constexpr TCHAR Name[] = TEXT("BatchQueryPresence");
+
+	struct Params
+	{
+		/** Local user performing the query */
+		FOnlineAccountIdHandle LocalUserId;
+		/** Users to query the presence for */
+		TArray<FOnlineAccountIdHandle> TargetUserIds;
+		/** If true, then future presence updates for these users will be propagated via the OnPresenceUpdated event */
+		bool bListenToChanges = true;
+	};
+
+	struct Result
+	{
+		/** The retrieved presence */
+		TArray<TSharedRef<const FUserPresence>> Presences;
+	};
+};
+
+struct FGetCachedPresence
+{
+	static constexpr TCHAR Name[] = TEXT("GetCachedPresence");
 
 	struct Params
 	{
@@ -84,13 +157,37 @@ struct FUpdatePresence
 	{
 		/** Local user performing the query */
 		FOnlineAccountIdHandle LocalUserId;
-		/** Mutations */
+		/** The new presence to send to the server */
+		TSharedRef<FUserPresence> Presence;
+	};
+
+	struct Result
+	{
+	};
+};
+
+struct FPartialUpdatePresence
+{
+	static constexpr TCHAR Name[] = TEXT("PartialUpdatePresence");
+
+	struct Params
+	{
+		/** Local user performing the query */
+		FOnlineAccountIdHandle LocalUserId;
+
+
 		struct FMutations
 		{
 			/** Presence state */
-			TOptional<EPresenceState> State;
+			TOptional<EUserPresenceStatus> Status;
+			/** Joinability state */
+			TOptional<EUserPresenceJoinability> Joinability;
+			/** Game status (playing/not playing this game) */
+			TOptional<EUserPresenceGameStatus> GameStatus;
 			/** Status string */
 			TOptional<FString> StatusString;
+			/** Rich presence string */
+			TOptional<FString> RichPresenceString;
 			/** Presence keys to update */
 			FPresenceProperties UpdatedProperties;
 			/** Properties to remove */
@@ -98,8 +195,8 @@ struct FUpdatePresence
 
 			FMutations& operator+=(FMutations&& NewParams);
 		};
-		/** Mutations */
 		FMutations Mutations;
+
 	};
 
 	struct Result
@@ -107,6 +204,45 @@ struct FUpdatePresence
 	};
 };
 
+inline FPartialUpdatePresence::Params::FMutations& FPartialUpdatePresence::Params::FMutations::operator+=(FPartialUpdatePresence::Params::FMutations&& NewParams)
+{
+	if (NewParams.Status.IsSet())
+	{
+		Status = NewParams.Status;
+	}
+
+	if (NewParams.Joinability.IsSet())
+	{
+		Joinability = NewParams.Joinability;
+	}
+
+	if (NewParams.GameStatus.IsSet())
+	{
+		GameStatus = NewParams.GameStatus;
+	}
+
+	if (NewParams.StatusString.IsSet())
+	{
+		StatusString = MoveTemp(NewParams.StatusString);
+	}
+
+	// Insert any updated keys / remove from removed
+	for (TPair<FString, FPresenceVariant>& NewUpdatedProperty : NewParams.UpdatedProperties)
+	{
+		RemovedProperties.Remove(NewUpdatedProperty.Key);
+		UpdatedProperties.Emplace(MoveTemp(NewUpdatedProperty.Key), MoveTemp(NewUpdatedProperty.Value));
+	}
+
+	// Merge any removed keys / remove from updated
+	RemovedProperties.Reserve(RemovedProperties.Num() + NewParams.RemovedProperties.Num());
+	for (FString& NewRemovedKey : NewParams.RemovedProperties)
+	{
+		UpdatedProperties.Remove(NewRemovedKey);
+		RemovedProperties.AddUnique(MoveTemp(NewRemovedKey));
+	}
+	return *this;
+}
+
 /** Struct for PresenceUpdated event */
 struct FPresenceUpdated
 {
@@ -126,15 +262,23 @@ public:
 	 * @return
 	 */
 	virtual TOnlineAsyncOpHandle<FQueryPresence> QueryPresence(FQueryPresence::Params&& Params) = 0;
+	
+	/**
+	 * Queries presences for all users in the array
+	 * 
+	 * @param Params for the BatchQueryPresence call
+	 * @return
+	 */
+	virtual TOnlineAsyncOpHandle<FBatchQueryPresence> BatchQueryPresence(FBatchQueryPresence::Params&& Params) = 0;
 
 	/**
 	 * Get the presence of a user
 	 * Presence typically comes from QueryPresence or push events from the online service
 	 *
-	 * @param Params for the GetPresence call
+	 * @param Params for the GetCachedPresence call
 	 * @return
 	 */
-	virtual TOnlineResult<FGetPresence> GetPresence(FGetPresence::Params&& Params) = 0;
+	virtual TOnlineResult<FGetCachedPresence> GetCachedPresence(FGetCachedPresence::Params&& Params) = 0;
 
 	/**
 	 * Update your presence
@@ -144,6 +288,14 @@ public:
 	 */
 	virtual TOnlineAsyncOpHandle<FUpdatePresence> UpdatePresence(FUpdatePresence::Params&& Params) = 0;
 
+	/**
+	 * Updates only the parts of the presence that are specified
+	 *
+	 * @param Params for the UpdatePresence call
+	 * @return
+	 */
+	virtual TOnlineAsyncOpHandle<FPartialUpdatePresence> PartialUpdatePresence(FPartialUpdatePresence::Params&& Params) = 0;
+
 	/**
 	 * Get the event that is triggered when presence is updated
 	 *
@@ -152,70 +304,70 @@ public:
 	virtual TOnlineEvent<void(const FPresenceUpdated&)> OnPresenceUpdated() = 0;
 };
 
-inline FUpdatePresence::Params::FMutations& FUpdatePresence::Params::FMutations::operator+=(FUpdatePresence::Params::FMutations&& NewParams)
-{
-	if (NewParams.State.IsSet())
-	{
-		State = NewParams.State;
-	}
-
-	if (NewParams.StatusString.IsSet())
-	{
-		StatusString = MoveTemp(NewParams.StatusString);
-	}
-
-	// Insert any updated keys / remove from removed
-	for (TPair<FString, FPresenceVariant>& NewUpdatedProperty : NewParams.UpdatedProperties)
-	{
-		RemovedProperties.Remove(NewUpdatedProperty.Key);
-		UpdatedProperties.Emplace(MoveTemp(NewUpdatedProperty.Key), MoveTemp(NewUpdatedProperty.Value));
-	}
-
-	// Merge any removed keys / remove from updated
-	RemovedProperties.Reserve(RemovedProperties.Num() + NewParams.RemovedProperties.Num());
-	for (FString& NewRemovedKey : NewParams.RemovedProperties)
-	{
-		UpdatedProperties.Remove(NewRemovedKey);
-		RemovedProperties.AddUnique(MoveTemp(NewRemovedKey));
-	}
-	return *this;
-}
-
 namespace Meta {
 // TODO: Move to Presence_Meta.inl file?
 
+BEGIN_ONLINE_STRUCT_META(FUserPresence)
+	ONLINE_STRUCT_FIELD(FUserPresence, UserId),
+	ONLINE_STRUCT_FIELD(FUserPresence, Status),
+	ONLINE_STRUCT_FIELD(FUserPresence, Joinability),
+	ONLINE_STRUCT_FIELD(FUserPresence, GameStatus),
+	ONLINE_STRUCT_FIELD(FUserPresence, StatusString),
+	ONLINE_STRUCT_FIELD(FUserPresence, RichPresenceString),
+	ONLINE_STRUCT_FIELD(FUserPresence, Properties)
+END_ONLINE_STRUCT_META()
+
 BEGIN_ONLINE_STRUCT_META(FQueryPresence::Params)
 	ONLINE_STRUCT_FIELD(FQueryPresence::Params, LocalUserId),
-	ONLINE_STRUCT_FIELD(FQueryPresence::Params, TargetUserId)
+	ONLINE_STRUCT_FIELD(FQueryPresence::Params, TargetUserId),
+	ONLINE_STRUCT_FIELD(FQueryPresence::Params, bListenToChanges)
+END_ONLINE_STRUCT_META()
+
+BEGIN_ONLINE_STRUCT_META(FBatchQueryPresence::Params)
+	ONLINE_STRUCT_FIELD(FBatchQueryPresence::Params, LocalUserId),
+	ONLINE_STRUCT_FIELD(FBatchQueryPresence::Params, TargetUserIds),
+	ONLINE_STRUCT_FIELD(FBatchQueryPresence::Params, bListenToChanges)
 END_ONLINE_STRUCT_META()
 
 BEGIN_ONLINE_STRUCT_META(FQueryPresence::Result)
 	ONLINE_STRUCT_FIELD(FQueryPresence::Result, Presence)
 END_ONLINE_STRUCT_META()
 
-BEGIN_ONLINE_STRUCT_META(FGetPresence::Params)
-	ONLINE_STRUCT_FIELD(FGetPresence::Params, LocalUserId)
+BEGIN_ONLINE_STRUCT_META(FBatchQueryPresence::Result)
+	ONLINE_STRUCT_FIELD(FBatchQueryPresence::Result, Presences)
 END_ONLINE_STRUCT_META()
 
-BEGIN_ONLINE_STRUCT_META(FGetPresence::Result)
-	ONLINE_STRUCT_FIELD(FGetPresence::Result, Presence)
+BEGIN_ONLINE_STRUCT_META(FGetCachedPresence::Params)
+	ONLINE_STRUCT_FIELD(FGetCachedPresence::Params, LocalUserId)
+END_ONLINE_STRUCT_META()
+
+BEGIN_ONLINE_STRUCT_META(FGetCachedPresence::Result)
+	ONLINE_STRUCT_FIELD(FGetCachedPresence::Result, Presence)
 END_ONLINE_STRUCT_META()
 
 BEGIN_ONLINE_STRUCT_META(FUpdatePresence::Params)
 	ONLINE_STRUCT_FIELD(FUpdatePresence::Params, LocalUserId),
-	ONLINE_STRUCT_FIELD(FUpdatePresence::Params, Mutations)
+	ONLINE_STRUCT_FIELD(FUpdatePresence::Params, Presence)
 END_ONLINE_STRUCT_META()
 
-BEGIN_ONLINE_STRUCT_META(FUpdatePresence::Params::FMutations)
-	ONLINE_STRUCT_FIELD(FUpdatePresence::Params::FMutations, State),
-	ONLINE_STRUCT_FIELD(FUpdatePresence::Params::FMutations, StatusString),
-	ONLINE_STRUCT_FIELD(FUpdatePresence::Params::FMutations, UpdatedProperties),
-	ONLINE_STRUCT_FIELD(FUpdatePresence::Params::FMutations, RemovedProperties)
+BEGIN_ONLINE_STRUCT_META(FUpdatePresence::Result)
 END_ONLINE_STRUCT_META()
 
-BEGIN_ONLINE_STRUCT_META(FUpdatePresence::Result)
+BEGIN_ONLINE_STRUCT_META(FPartialUpdatePresence::Params)
+	ONLINE_STRUCT_FIELD(FPartialUpdatePresence::Params, LocalUserId),
+	ONLINE_STRUCT_FIELD(FPartialUpdatePresence::Params, Mutations)
 END_ONLINE_STRUCT_META()
 
+BEGIN_ONLINE_STRUCT_META(FPartialUpdatePresence::Params::FMutations)
+	ONLINE_STRUCT_FIELD(FPartialUpdatePresence::Params::FMutations, Status),
+	ONLINE_STRUCT_FIELD(FPartialUpdatePresence::Params::FMutations, Joinability),
+	ONLINE_STRUCT_FIELD(FPartialUpdatePresence::Params::FMutations, GameStatus),
+	ONLINE_STRUCT_FIELD(FPartialUpdatePresence::Params::FMutations, StatusString),
+	ONLINE_STRUCT_FIELD(FPartialUpdatePresence::Params::FMutations, UpdatedProperties),
+	ONLINE_STRUCT_FIELD(FPartialUpdatePresence::Params::FMutations, RemovedProperties)
+END_ONLINE_STRUCT_META()
+
+
 BEGIN_ONLINE_STRUCT_META(FPresenceUpdated)
 	ONLINE_STRUCT_FIELD(FPresenceUpdated, LocalUserId),
 	ONLINE_STRUCT_FIELD(FPresenceUpdated, UpdatedPresence)
@@ -224,3 +376,5 @@ END_ONLINE_STRUCT_META()
 /* Meta*/ }
 
 /* UE::Online */ }
+
+#undef LOCTEXT_NAMESPACE
\ No newline at end of file
diff --git a/Engine/Plugins/Online/Experimental/OnlineServicesEOS/Source/Private/Online/PresenceEOS.cpp b/Engine/Plugins/Online/Experimental/OnlineServicesEOS/Source/Private/Online/PresenceEOS.cpp
--- a/Engine/Plugins/Online/Experimental/OnlineServicesEOS/Source/Private/Online/PresenceEOS.cpp
+++ b/Engine/Plugins/Online/Experimental/OnlineServicesEOS/Source/Private/Online/PresenceEOS.cpp
@@ -6,52 +6,53 @@
 #include "Online/OnlineServicesEOS.h"
 #include "Online/OnlineServicesEOSTypes.h"
 #include "Online/AuthEOS.h"
+#include "Online/ErrorsEOS.h"
 
 #include "eos_presence.h"
 
 namespace UE::Online {
 
-static inline EPresenceState ToEPresenceState(EOS_Presence_EStatus InStatus)
+static inline EUserPresenceStatus ToEPresenceState(EOS_Presence_EStatus InStatus)
 {
 	switch (InStatus)
 	{
 		case EOS_Presence_EStatus::EOS_PS_Online:
 		{
-			return EPresenceState::Online;
+			return EUserPresenceStatus::Online;
 		}
 		case EOS_Presence_EStatus::EOS_PS_Away:
 		{
-			return EPresenceState::Away;
+			return EUserPresenceStatus::Away;
 		}
 		case EOS_Presence_EStatus::EOS_PS_ExtendedAway:
 		{
-			return EPresenceState::ExtendedAway;
+			return EUserPresenceStatus::ExtendedAway;
 		}
 		case EOS_Presence_EStatus::EOS_PS_DoNotDisturb:
 		{
-			return EPresenceState::DoNotDisturb;
+			return EUserPresenceStatus::DoNotDisturb;
 		}
 	}
-	return EPresenceState::Offline;
+	return EUserPresenceStatus::Offline;
 }
 
-static inline EOS_Presence_EStatus ToEOS_Presence_EStatus(EPresenceState InState)
+static inline EOS_Presence_EStatus ToEOS_Presence_EStatus(EUserPresenceStatus InState)
 {
 	switch (InState)
 	{
-	case EPresenceState::Online:
+	case EUserPresenceStatus::Online:
 	{
 		return EOS_Presence_EStatus::EOS_PS_Online;
 	}
-	case EPresenceState::Away:
+	case EUserPresenceStatus::Away:
 	{
 		return EOS_Presence_EStatus::EOS_PS_Away;
 	}
-	case EPresenceState::ExtendedAway:
+	case EUserPresenceStatus::ExtendedAway:
 	{
 		return EOS_Presence_EStatus::EOS_PS_ExtendedAway;
 	}
-	case EPresenceState::DoNotDisturb:
+	case EUserPresenceStatus::DoNotDisturb:
 	{
 		return EOS_Presence_EStatus::EOS_PS_DoNotDisturb;
 	}
@@ -102,15 +103,13 @@ TOnlineAsyncOpHandle<FQueryPresence> FPresenceEOS::QueryPresence(FQueryPresence:
 		const FQueryPresence::Params& Params = Op->GetParams();
 		if (!Services.Get<FAuthEOS>()->IsLoggedIn(Params.LocalUserId))
 		{
-			// TODO: Error codes
-			Op->SetError(Errors::Unknown());
+			Op->SetError(Errors::NotLoggedIn());
 			return Op->GetHandle();
 		}
 		EOS_EpicAccountId TargetUserEasId = GetEpicAccountId(Params.TargetUserId);
 		if (!EOS_EpicAccountId_IsValid(TargetUserEasId))
 		{
-			// TODO: Error codes
-			Op->SetError(Errors::Unknown());
+			Op->SetError(Errors::NotLoggedIn());
 			return Op->GetHandle();
 		}
 
@@ -118,7 +117,7 @@ TOnlineAsyncOpHandle<FQueryPresence> FPresenceEOS::QueryPresence(FQueryPresence:
 		const bool bIsLocalUser = Services.Get<FAuthEOS>()->IsLoggedIn(Params.TargetUserId);
 		if (bIsLocalUser)
 		{
-			Op->SetError(Errors::Unknown()); // TODO: Error codes
+			Op->SetError(Errors::CannotQueryLocalUsers()); 
 			return Op->GetHandle();
 		}
 
@@ -145,7 +144,7 @@ TOnlineAsyncOpHandle<FQueryPresence> FPresenceEOS::QueryPresence(FQueryPresence:
 				}
 				else
 				{
-					InAsyncOp.SetError(Errors::Unknown()); // TODO: Error codes
+					InAsyncOp.SetError(Errors::FromEOSResult(Data->ResultCode)); 
 				}
 			})
 			.Enqueue(GetSerialQueue());
@@ -153,21 +152,189 @@ TOnlineAsyncOpHandle<FQueryPresence> FPresenceEOS::QueryPresence(FQueryPresence:
 	return Op->GetHandle();
 }
 
-TOnlineResult<FGetPresence> FPresenceEOS::GetPresence(FGetPresence::Params&& Params)
+TOnlineResult<FGetCachedPresence> FPresenceEOS::GetCachedPresence(FGetCachedPresence::Params&& Params)
 {
 	if (TMap<FOnlineAccountIdHandle, TSharedRef<FUserPresence>>* PresenceList = PresenceLists.Find(Params.LocalUserId))
 	{
 		TSharedRef<FUserPresence>* PresencePtr = PresenceList->Find(Params.TargetUserId);
 		if (PresencePtr)
 		{
-			FGetPresence::Result Result = { *PresencePtr };
-			return TOnlineResult<FGetPresence>(MoveTemp(Result));
+			FGetCachedPresence::Result Result = { *PresencePtr };
+			return TOnlineResult<FGetCachedPresence>(MoveTemp(Result));
 		}
 	}
-	return TOnlineResult<FGetPresence>(Errors::Unknown()); // TODO: error codes
+	return TOnlineResult<FGetCachedPresence>(Errors::NotFound()); 
 }
 
+
 TOnlineAsyncOpHandle<FUpdatePresence> FPresenceEOS::UpdatePresence(FUpdatePresence::Params&& InParams)
+{
+	TOnlineAsyncOpRef<FUpdatePresence> Op = GetOp<FUpdatePresence>(MoveTemp(InParams));
+	const FUpdatePresence::Params& Params = Op->GetParams();
+	if (!Services.Get<FAuthEOS>()->IsLoggedIn(Params.LocalUserId))
+	{
+		Op->SetError(Errors::NotLoggedIn());
+		return Op->GetHandle();
+	}
+
+	Op->Then([this](TOnlineAsyncOp<FUpdatePresence>& InAsyncOp) mutable
+	{
+		const FUpdatePresence::Params& Params = InAsyncOp.GetParams();
+		EOS_HPresenceModification ChangeHandle = nullptr;
+		EOS_Presence_CreatePresenceModificationOptions Options = { };
+		Options.ApiVersion = EOS_PRESENCE_CREATEPRESENCEMODIFICATION_API_LATEST;
+		Options.LocalUserId = GetEpicAccountIdChecked(Params.LocalUserId);
+		EOS_EResult CreatePresenceModificationResult = EOS_Presence_CreatePresenceModification(PresenceHandle, &Options, &ChangeHandle);
+		if (CreatePresenceModificationResult == EOS_EResult::EOS_Success)
+		{
+			check(ChangeHandle != nullptr);
+
+			// State
+			EOS_PresenceModification_SetStatusOptions StatusOptions = { };
+			StatusOptions.ApiVersion = EOS_PRESENCEMODIFICATION_SETSTATUS_API_LATEST;
+			StatusOptions.Status = ToEOS_Presence_EStatus(Params.Presence->Status);
+			EOS_EResult SetStatusResult = EOS_PresenceModification_SetStatus(ChangeHandle, &StatusOptions);
+			if (SetStatusResult != EOS_EResult::EOS_Success)
+			{
+				UE_LOG(LogTemp, Warning, TEXT("UpdatePresence: EOS_PresenceModification_SetStatus failed with result %s"), *LexToString(SetStatusResult));
+				InAsyncOp.SetError(Errors::FromEOSResult(SetStatusResult));
+				return MakeFulfilledPromise<const EOS_Presence_SetPresenceCallbackInfo*>(nullptr).GetFuture();
+			}
+
+			// Raw rich text
+			// Convert the status string as the rich text string
+			EOS_PresenceModification_SetRawRichTextOptions RawRichTextOptions = { };
+			RawRichTextOptions.ApiVersion = EOS_PRESENCEMODIFICATION_SETRAWRICHTEXT_API_LATEST;
+
+			FTCHARToUTF8 Utf8RawRichText(*Params.Presence->StatusString);
+			RawRichTextOptions.RichText = Utf8RawRichText.Get();
+
+			EOS_EResult SetRawRichTextResult = EOS_PresenceModification_SetRawRichText(ChangeHandle, &RawRichTextOptions);
+			if (SetRawRichTextResult != EOS_EResult::EOS_Success)
+			{
+				UE_LOG(LogTemp, Warning, TEXT("UpdatePresence: EOS_PresenceModification_SetRawRichText failed with result %s"), *LexToString(SetRawRichTextResult));
+				InAsyncOp.SetError(Errors::FromEOSResult(SetRawRichTextResult));
+				return MakeFulfilledPromise<const EOS_Presence_SetPresenceCallbackInfo*>(nullptr).GetFuture();
+			}
+
+			// EOS needs to be specific on which fields are removed and which aren't, so grab the last presence and check which are deletions
+			TArray<FString> RemovedProperties;
+			TSharedRef<FUserPresence> LastUserPresence = FindOrCreatePresence(Params.LocalUserId, Params.LocalUserId);
+			for (const TPair<FString, FPresenceVariant>& Pair : LastUserPresence->Properties)
+			{
+				if(!Params.Presence->Properties.Contains(Pair.Key))
+				{
+					RemovedProperties.Add(Pair.Key);
+				}
+			}
+
+			// Removed fields
+			if (RemovedProperties.Num() > 0)
+			{
+				// EOS_PresenceModification_DeleteData
+				TArray<FTCHARToUTF8, TInlineAllocator<EOS_PRESENCE_DATA_MAX_KEYS>> Utf8Strings;
+				TArray<EOS_PresenceModification_DataRecordId, TInlineAllocator<EOS_PRESENCE_DATA_MAX_KEYS>> RecordIds;
+				int32 CurrentIndex = 0;
+				for (const FString& RemovedProperty : RemovedProperties)
+				{
+					const FTCHARToUTF8& Utf8Key = Utf8Strings.Emplace_GetRef(RemovedProperty);
+
+					EOS_PresenceModification_DataRecordId& RecordId = RecordIds.Emplace_GetRef();
+					RecordId.ApiVersion = EOS_PRESENCEMODIFICATION_DATARECORDID_API_LATEST;
+					RecordId.Key = Utf8Key.Get();
+
+					UE_LOG(LogTemp, Warning, TEXT("UpdatePresence: Removing field %s"), *RemovedProperty); // Temp logging
+				}
+
+				EOS_PresenceModification_DeleteDataOptions DataOptions = { };
+				DataOptions.ApiVersion = EOS_PRESENCEMODIFICATION_DELETEDATA_API_LATEST;
+				DataOptions.RecordsCount = RecordIds.Num();
+				DataOptions.Records = RecordIds.GetData();
+				EOS_EResult DeleteDataResult = EOS_PresenceModification_DeleteData(ChangeHandle, &DataOptions);
+				if (DeleteDataResult != EOS_EResult::EOS_Success)
+				{
+					UE_LOG(LogTemp, Warning, TEXT("UpdatePresence: EOS_PresenceModification_DeleteDataOptions failed with result %s"), *LexToString(DeleteDataResult));
+					InAsyncOp.SetError(Errors::FromEOSResult(DeleteDataResult));
+					return MakeFulfilledPromise<const EOS_Presence_SetPresenceCallbackInfo*>(nullptr).GetFuture();
+				}
+			}
+
+			// Added/Updated fields
+			if (Params.Presence->Properties.Num() > 0)
+			{
+				if (Params.Presence->Properties.Num() > EOS_PRESENCE_DATA_MAX_KEYS)
+				{
+					// TODO: Move this check higher.  Needs to take into account number of present fields (not just ones updated) and removed fields.
+					UE_LOG(LogTemp, Warning, TEXT("UpdatePresence: Too many presence keys.  %u/%u"), Params.Presence->Properties.Num(), EOS_PRESENCE_DATA_MAX_KEYS);
+					InAsyncOp.SetError(Errors::InvalidParams());
+					return MakeFulfilledPromise<const EOS_Presence_SetPresenceCallbackInfo*>(nullptr).GetFuture();
+				}
+				TArray<FTCHARToUTF8, TInlineAllocator<EOS_PRESENCE_DATA_MAX_KEYS * 2>> Utf8Strings;
+				TArray<EOS_Presence_DataRecord, TInlineAllocator<EOS_PRESENCE_DATA_MAX_KEYS>> Records;
+				int32 CurrentIndex = 0;
+				for (const TPair<FString, FPresenceVariant>& UpdatedProperty : Params.Presence->Properties)
+				{
+					const FTCHARToUTF8& Utf8Key = Utf8Strings.Emplace_GetRef(UpdatedProperty.Key);
+					const FTCHARToUTF8& Utf8Value = Utf8Strings.Emplace_GetRef(UpdatedProperty.Value); // TODO: Better serialization
+
+					EOS_Presence_DataRecord& Record = Records.Emplace_GetRef();
+					Record.ApiVersion = EOS_PRESENCE_DATARECORD_API_LATEST;
+					Record.Key = Utf8Key.Get();
+					Record.Value = Utf8Value.Get();
+					UE_LOG(LogTemp, Warning, TEXT("UpdatePresence: Set field [%s] to [%s]"), *UpdatedProperty.Key, *UpdatedProperty.Value); // Temp logging
+				}
+
+				EOS_PresenceModification_SetDataOptions DataOptions = { };
+				DataOptions.ApiVersion = EOS_PRESENCEMODIFICATION_SETDATA_API_LATEST;
+				DataOptions.RecordsCount = Records.Num();
+				DataOptions.Records = Records.GetData();
+				EOS_EResult SetDataResult = EOS_PresenceModification_SetData(ChangeHandle, &DataOptions);
+				if (SetDataResult != EOS_EResult::EOS_Success)
+				{
+					UE_LOG(LogTemp, Warning, TEXT("UpdatePresence: EOS_PresenceModification_SetData failed with result %s"), *LexToString(SetDataResult));
+					InAsyncOp.SetError(Errors::FromEOSResult(SetDataResult));
+					return MakeFulfilledPromise<const EOS_Presence_SetPresenceCallbackInfo*>(nullptr).GetFuture();
+				}
+			}
+
+			EOS_Presence_SetPresenceOptions SetPresenceOptions = { };
+			SetPresenceOptions.ApiVersion = EOS_PRESENCE_SETPRESENCE_API_LATEST;
+			SetPresenceOptions.LocalUserId = GetEpicAccountIdChecked(Params.LocalUserId);
+			SetPresenceOptions.PresenceModificationHandle = ChangeHandle;
+			TFuture<const EOS_Presence_SetPresenceCallbackInfo*> EOSAsyncFuture = EOS_Async<EOS_Presence_SetPresenceCallbackInfo>(EOS_Presence_SetPresence, PresenceHandle, SetPresenceOptions);
+			EOS_PresenceModification_Release(ChangeHandle);
+			return EOSAsyncFuture;
+		}
+		else
+		{
+			InAsyncOp.SetError(Errors::FromEOSResult(CreatePresenceModificationResult));
+			return MakeFulfilledPromise<const EOS_Presence_SetPresenceCallbackInfo*>(nullptr).GetFuture();
+		}
+	})
+	.Then([this](TOnlineAsyncOp<FUpdatePresence>& InAsyncOp, const EOS_Presence_SetPresenceCallbackInfo* Data) mutable
+	{
+		UE_LOG(LogTemp, Warning, TEXT("SetPresenceResult: [%s]"), *LexToString(Data->ResultCode));
+
+		if (Data->ResultCode == EOS_EResult::EOS_Success)
+		{
+			// Update local presence
+			const FUpdatePresence::Params& Params = InAsyncOp.GetParams();
+			TSharedRef<FUserPresence> LocalUserPresence = FindOrCreatePresence(Params.LocalUserId, Params.LocalUserId);
+			*LocalUserPresence = *Params.Presence;
+
+			InAsyncOp.SetResult(FUpdatePresence::Result());
+		}
+		else
+		{
+			InAsyncOp.SetError(Errors::FromEOSResult(Data->ResultCode));
+		}
+	})
+	.Enqueue(GetSerialQueue());
+
+	return Op->GetHandle();
+}
+
+TOnlineAsyncOpHandle<FPartialUpdatePresence> FPresenceEOS::PartialUpdatePresence(FPartialUpdatePresence::Params&& InParams)
 {
 	// TODO: Validate params
 	// EOS_PRESENCE_DATA_MAX_KEYS - number of total keys.  Compare proposed with existing with pending ops?  Include removed keys!
@@ -175,22 +342,21 @@ TOnlineAsyncOpHandle<FUpdatePresence> FPresenceEOS::UpdatePresence(FUpdatePresen
 	// EOS_PRESENCE_DATA_MAX_VALUE_LENGTH - length of each value.  Compare updated with this.
 	// EOS_PRESENCE_RICH_TEXT_MAX_VALUE_LENGTH - length of status. Compare updated with this.
 
-	TOnlineAsyncOpRef<FUpdatePresence> Op = GetMergeableOp<FUpdatePresence>(MoveTemp(InParams));
+	TOnlineAsyncOpRef<FPartialUpdatePresence> Op = GetMergeableOp<FPartialUpdatePresence>(MoveTemp(InParams));
 	if (!Op->IsReady())
 	{
 		// Initialize
-		const FUpdatePresence::Params& Params = Op->GetParams();
+		const FPartialUpdatePresence::Params& Params = Op->GetParams();
 		if (!Services.Get<FAuthEOS>()->IsLoggedIn(Params.LocalUserId))
 		{
-			// TODO: Error codes
-			Op->SetError(Errors::Unknown());
+			Op->SetError(Errors::NotLoggedIn());
 			return Op->GetHandle();
 		}
 
 		// Don't cache anything from Params as they could be modified by another merge in the meanwhile.
-		Op->Then([this](TOnlineAsyncOp<FUpdatePresence>& InAsyncOp) mutable
+		Op->Then([this](TOnlineAsyncOp<FPartialUpdatePresence>& InAsyncOp) mutable
 			{
-				const FUpdatePresence::Params& Params = InAsyncOp.GetParams();
+				const FPartialUpdatePresence::Params& Params = InAsyncOp.GetParams();
 				EOS_HPresenceModification ChangeHandle = nullptr;
 				EOS_Presence_CreatePresenceModificationOptions Options = { };
 				Options.ApiVersion = EOS_PRESENCE_CREATEPRESENCEMODIFICATION_API_LATEST;
@@ -201,21 +367,21 @@ TOnlineAsyncOpHandle<FUpdatePresence> FPresenceEOS::UpdatePresence(FUpdatePresen
 					check(ChangeHandle != nullptr);
 
 					// State
-					if (Params.Mutations.State.IsSet())
+					if (Params.Mutations.Status.IsSet())
 					{
 						EOS_PresenceModification_SetStatusOptions StatusOptions = { };
 						StatusOptions.ApiVersion = EOS_PRESENCEMODIFICATION_SETSTATUS_API_LATEST;
-						StatusOptions.Status = ToEOS_Presence_EStatus(Params.Mutations.State.GetValue());
+						StatusOptions.Status = ToEOS_Presence_EStatus(Params.Mutations.Status.GetValue());
 						EOS_EResult SetStatusResult = EOS_PresenceModification_SetStatus(ChangeHandle, &StatusOptions);
 						if (SetStatusResult != EOS_EResult::EOS_Success)
 						{
 							UE_LOG(LogTemp, Warning, TEXT("UpdatePresence: EOS_PresenceModification_SetStatus failed with result %s"), *LexToString(SetStatusResult));
-							InAsyncOp.SetError(Errors::Unknown());
+							InAsyncOp.SetError(Errors::FromEOSResult(SetStatusResult));
 							return MakeFulfilledPromise<const EOS_Presence_SetPresenceCallbackInfo*>(nullptr).GetFuture();
 						}
 						else
 						{
-							UE_LOG(LogTemp, Warning, TEXT("UpdatePresence: Status set to %u"), (uint8)Params.Mutations.State.GetValue()); // Temp logging
+							UE_LOG(LogTemp, Warning, TEXT("UpdatePresence: Status set to %u"), (uint8)Params.Mutations.Status.GetValue()); // Temp logging
 						}
 					}
 
@@ -233,7 +399,7 @@ TOnlineAsyncOpHandle<FUpdatePresence> FPresenceEOS::UpdatePresence(FUpdatePresen
 						if (SetRawRichTextResult != EOS_EResult::EOS_Success)
 						{
 							UE_LOG(LogTemp, Warning, TEXT("UpdatePresence: EOS_PresenceModification_SetRawRichText failed with result %s"), *LexToString(SetRawRichTextResult));
-							InAsyncOp.SetError(Errors::Unknown());
+							InAsyncOp.SetError(Errors::FromEOSResult(SetRawRichTextResult));
 							return MakeFulfilledPromise<const EOS_Presence_SetPresenceCallbackInfo*>(nullptr).GetFuture();
 						}
 						else
@@ -241,7 +407,6 @@ TOnlineAsyncOpHandle<FUpdatePresence> FPresenceEOS::UpdatePresence(FUpdatePresen
 							UE_LOG(LogTemp, Warning, TEXT("UpdatePresence: RichText set to %s"), *Params.Mutations.StatusString.GetValue()); // Temp logging
 						}
 					}
-
 					// Removed fields
 					if (Params.Mutations.RemovedProperties.Num() > 0)
 					{
@@ -268,7 +433,7 @@ TOnlineAsyncOpHandle<FUpdatePresence> FPresenceEOS::UpdatePresence(FUpdatePresen
 						if (DeleteDataResult != EOS_EResult::EOS_Success)
 						{
 							UE_LOG(LogTemp, Warning, TEXT("UpdatePresence: EOS_PresenceModification_DeleteDataOptions failed with result %s"), *LexToString(DeleteDataResult));
-							InAsyncOp.SetError(Errors::Unknown());
+							InAsyncOp.SetError(Errors::FromEOSResult(DeleteDataResult));
 							return MakeFulfilledPromise<const EOS_Presence_SetPresenceCallbackInfo*>(nullptr).GetFuture();
 						}
 					}
@@ -280,7 +445,7 @@ TOnlineAsyncOpHandle<FUpdatePresence> FPresenceEOS::UpdatePresence(FUpdatePresen
 						{
 							// TODO: Move this check higher.  Needs to take into account number of present fields (not just ones updated) and removed fields.
 							UE_LOG(LogTemp, Warning, TEXT("UpdatePresence: Too many presence keys.  %u/%u"), Params.Mutations.UpdatedProperties.Num(), EOS_PRESENCE_DATA_MAX_KEYS);
-							InAsyncOp.SetError(Errors::Unknown());
+							InAsyncOp.SetError(Errors::InvalidParams());
 							return MakeFulfilledPromise<const EOS_Presence_SetPresenceCallbackInfo*>(nullptr).GetFuture();
 						}
 						TArray<FTCHARToUTF8, TInlineAllocator<EOS_PRESENCE_DATA_MAX_KEYS * 2>> Utf8Strings;
@@ -306,7 +471,7 @@ TOnlineAsyncOpHandle<FUpdatePresence> FPresenceEOS::UpdatePresence(FUpdatePresen
 						if (SetDataResult != EOS_EResult::EOS_Success)
 						{
 							UE_LOG(LogTemp, Warning, TEXT("UpdatePresence: EOS_PresenceModification_SetData failed with result %s"), *LexToString(SetDataResult));
-							InAsyncOp.SetError(Errors::Unknown());
+							InAsyncOp.SetError(Errors::FromEOSResult(SetDataResult));
 							return MakeFulfilledPromise<const EOS_Presence_SetPresenceCallbackInfo*>(nullptr).GetFuture();
 						}
 					}
@@ -321,28 +486,35 @@ TOnlineAsyncOpHandle<FUpdatePresence> FPresenceEOS::UpdatePresence(FUpdatePresen
 				}
 				else
 				{
-					// TODO:  Error code
-					InAsyncOp.SetError(Errors::Unknown());
+					InAsyncOp.SetError(Errors::FromEOSResult(CreatePresenceModificationResult));
 					return MakeFulfilledPromise<const EOS_Presence_SetPresenceCallbackInfo*>(nullptr).GetFuture();
 				}
 			})
-			.Then([this](TOnlineAsyncOp<FUpdatePresence>& InAsyncOp, const EOS_Presence_SetPresenceCallbackInfo* Data) mutable
+			.Then([this](TOnlineAsyncOp<FPartialUpdatePresence>& InAsyncOp, const EOS_Presence_SetPresenceCallbackInfo* Data) mutable
 			{
 				UE_LOG(LogTemp, Warning, TEXT("SetPresenceResult: [%s]"), *LexToString(Data->ResultCode));
 
 				if (Data->ResultCode == EOS_EResult::EOS_Success)
 				{
 					// Update local presence
-					const FUpdatePresence::Params& Params = InAsyncOp.GetParams();
+					const FPartialUpdatePresence::Params& Params = InAsyncOp.GetParams();
 					TSharedRef<FUserPresence> LocalUserPresence = FindOrCreatePresence(Params.LocalUserId, Params.LocalUserId);
-					if (Params.Mutations.State.IsSet())
+					if (Params.Mutations.Status.IsSet())
 					{
-						LocalUserPresence->State = Params.Mutations.State.GetValue();
+						LocalUserPresence->Status = Params.Mutations.Status.GetValue();
 					}
 					if (Params.Mutations.StatusString.IsSet())
 					{
 						LocalUserPresence->StatusString = Params.Mutations.StatusString.GetValue();
 					}
+					if (Params.Mutations.GameStatus.IsSet())
+					{
+						LocalUserPresence->GameStatus = Params.Mutations.GameStatus.GetValue();
+					}
+					if (Params.Mutations.Joinability.IsSet())
+					{
+						LocalUserPresence->Joinability = Params.Mutations.Joinability.GetValue();
+					}
 					for (const FString& RemovedKey : Params.Mutations.RemovedProperties)
 					{
 						LocalUserPresence->Properties.Remove(RemovedKey);
@@ -352,11 +524,11 @@ TOnlineAsyncOpHandle<FUpdatePresence> FPresenceEOS::UpdatePresence(FUpdatePresen
 						LocalUserPresence->Properties.Emplace(UpdatedProperty.Key, UpdatedProperty.Value);
 					}
 
-					InAsyncOp.SetResult(FUpdatePresence::Result());
+					InAsyncOp.SetResult(FPartialUpdatePresence::Result());
 				}
 				else
 				{
-					InAsyncOp.SetError(Errors::Unknown()); // TODO: Error codes
+					InAsyncOp.SetError(Errors::FromEOSResult(Data->ResultCode)); 
 				}
 			})
 			.Enqueue(GetSerialQueue());
@@ -395,11 +567,11 @@ void FPresenceEOS::UpdateUserPresence(FOnlineAccountIdHandle LocalUserId, FOnlin
 	if (CopyPresenceResult == EOS_EResult::EOS_Success)
 	{
 		// Convert the presence data to our format
-		EPresenceState NewPresenceState = ToEPresenceState(PresenceInfo->Status);
-		if (UserPresence->State != NewPresenceState)
+		EUserPresenceStatus NewPresenceState = ToEPresenceState(PresenceInfo->Status);
+		if (UserPresence->Status != NewPresenceState)
 		{
 			bPresenceHasChanged = true;
-			UserPresence->State = NewPresenceState;
+			UserPresence->Status = NewPresenceState;
 		}
 
 		FString NewStatusString = UTF8_TO_TCHAR(PresenceInfo->RichText);
diff --git a/Engine/Plugins/Online/Experimental/OnlineServicesEOS/Source/Public/Online/PresenceEOS.h b/Engine/Plugins/Online/Experimental/OnlineServicesEOS/Source/Public/Online/PresenceEOS.h
--- a/Engine/Plugins/Online/Experimental/OnlineServicesEOS/Source/Public/Online/PresenceEOS.h
+++ b/Engine/Plugins/Online/Experimental/OnlineServicesEOS/Source/Public/Online/PresenceEOS.h
@@ -23,8 +23,9 @@ public:
 	virtual void PreShutdown() override;
 
 	virtual TOnlineAsyncOpHandle<FQueryPresence> QueryPresence(FQueryPresence::Params&& Params) override;
-	virtual TOnlineResult<FGetPresence> GetPresence(FGetPresence::Params&& Params) override;
+	virtual TOnlineResult<FGetCachedPresence> GetCachedPresence(FGetCachedPresence::Params&& Params) override;
 	virtual TOnlineAsyncOpHandle<FUpdatePresence> UpdatePresence(FUpdatePresence::Params&& Params) override;
+	virtual TOnlineAsyncOpHandle<FPartialUpdatePresence> PartialUpdatePresence(FPartialUpdatePresence::Params&& Params) override;
 
 protected:
 	/** Get a user's presence, creating entries if missing */
