- Fix unsafe EOS_Async usage, binding to the future _after_ the call is made. In situations where the SDK immediately fires the callback, this causes us to bind the continuation _after_ the SDK function has returned, and the CallbackInfo has become invalidated.

diff --git a/Engine/Plugins/Online/Experimental/OnlineServicesEOS/Source/Private/Online/AuthEOS.cpp b/Engine/Plugins/Online/Experimental/OnlineServicesEOS/Source/Private/Online/AuthEOS.cpp
--- a/Engine/Plugins/Online/Experimental/OnlineServicesEOS/Source/Private/Online/AuthEOS.cpp
+++ b/Engine/Plugins/Online/Experimental/OnlineServicesEOS/Source/Private/Online/AuthEOS.cpp
@@ -96,15 +96,15 @@ TOnlineAsyncOpHandle<FAuthLogin> FAuthEOS::Login(FAuthLogin::Params&& Params)
 		return ConnectLoginOp.GetOwningOperation().GetHandle();
 	}
 
-	ConnectLoginOp.Then([this](TOnlineAsyncOp<FAuthLogin>& InAsyncOp, TSharedPtr<FEOSConnectLoginCredentials>&& InConnectLoginCredentials)
+	ConnectLoginOp.Then([this](TOnlineAsyncOp<FAuthLogin>& InAsyncOp, TSharedPtr<FEOSConnectLoginCredentials>&& InConnectLoginCredentials, TPromise<const EOS_Connect_LoginCallbackInfo*>&& Promise)
 	{
 		EOS_Connect_LoginOptions ConnectLoginOptions = { };
 		ConnectLoginOptions.ApiVersion = EOS_CONNECT_LOGIN_API_LATEST;
 		ConnectLoginOptions.Credentials = InConnectLoginCredentials.Get();
 
-		return EOS_Async<EOS_Connect_LoginCallbackInfo>(EOS_Connect_Login, ConnectHandle, ConnectLoginOptions);
+		EOS_Async(EOS_Connect_Login, ConnectHandle, ConnectLoginOptions, MoveTemp(Promise));
 	})
-	.Then([this](TOnlineAsyncOp<FAuthLogin>& InAsyncOp, const EOS_Connect_LoginCallbackInfo* Data)
+	.Then([this](TOnlineAsyncOp<FAuthLogin>& InAsyncOp, const EOS_Connect_LoginCallbackInfo* Data, TPromise<const EOS_Connect_CreateUserCallbackInfo*>&& Promise)
 	{
 		UE_LOG(LogTemp, Verbose, TEXT("[FAuthEOS::Login] EOS_Connect_Login Result: [%s]"), *LexToString(Data->ResultCode));
 
@@ -121,7 +121,8 @@ TOnlineAsyncOpHandle<FAuthLogin> FAuthEOS::Login(FAuthLogin::Params&& Params)
 			ConnectCreateUserOptions.ApiVersion = EOS_CONNECT_CREATEUSER_API_LATEST;
 			ConnectCreateUserOptions.ContinuanceToken = Data->ContinuanceToken;
 
-			return EOS_Async<EOS_Connect_CreateUserCallbackInfo>(EOS_Connect_CreateUser, ConnectHandle, ConnectCreateUserOptions);
+			EOS_Async(EOS_Connect_CreateUser, ConnectHandle, ConnectCreateUserOptions, MoveTemp(Promise));
+			return;
 		}
 		else
 		{
@@ -129,7 +130,7 @@ TOnlineAsyncOpHandle<FAuthLogin> FAuthEOS::Login(FAuthLogin::Params&& Params)
 			InAsyncOp.SetError(Errors::Unknown()); // TODO
 		}
 
-		return MakeFulfilledPromise<const EOS_Connect_CreateUserCallbackInfo*>().GetFuture();
+		Promise.EmplaceValue();
 	})
 	.Then([this](TOnlineAsyncOp<FAuthLogin>& InAsyncOp, const EOS_Connect_CreateUserCallbackInfo* Data)
 	{
@@ -259,12 +260,12 @@ TOnlineChainableAsyncOp<FAuthLogin, TSharedPtr<FEOSConnectLoginCredentials>> FAu
 		return InAsyncOp.Then(FailureLambda);
 	}
 
-	return InAsyncOp.Then([this, LoginOptions, Credentials](TOnlineAsyncOp<FAuthLogin>& InAsyncOp) mutable
+	return InAsyncOp.Then([this, LoginOptions, Credentials](TOnlineAsyncOp<FAuthLogin>& InAsyncOp, TPromise<const EOS_Auth_LoginCallbackInfo*>&& Promise) mutable
 	{
 		LoginOptions.Credentials = &Credentials;
-		return EOS_Async<EOS_Auth_LoginCallbackInfo>(EOS_Auth_Login, AuthHandle, LoginOptions);
+		EOS_Async(EOS_Auth_Login, AuthHandle, LoginOptions, MoveTemp(Promise));
 	})
-	.Then([this](TOnlineAsyncOp<FAuthLogin>& InAsyncOp, const EOS_Auth_LoginCallbackInfo* Data)
+	.Then([this](TOnlineAsyncOp<FAuthLogin>& InAsyncOp, const EOS_Auth_LoginCallbackInfo* Data, TPromise<const EOS_Auth_LinkAccountCallbackInfo*>&& Promise)
 	{
 		UE_LOG(LogTemp, Verbose, TEXT("[FAuthEOS::Login] EOS_Auth_Login Result: [%s]"), *LexToString(Data->ResultCode));
 
@@ -281,7 +282,8 @@ TOnlineChainableAsyncOp<FAuthLogin, TSharedPtr<FEOSConnectLoginCredentials>> FAu
 			LinkAccountOptions.ApiVersion = EOS_AUTH_LINKACCOUNT_API_LATEST;
 			LinkAccountOptions.ContinuanceToken = Data->ContinuanceToken;
 
-			return EOS_Async<EOS_Auth_LinkAccountCallbackInfo>(EOS_Auth_LinkAccount, AuthHandle, LinkAccountOptions);
+			EOS_Async(EOS_Auth_LinkAccount, AuthHandle, LinkAccountOptions, MoveTemp(Promise));
+			return;
 		}
 		else
 		{
@@ -294,7 +296,7 @@ TOnlineChainableAsyncOp<FAuthLogin, TSharedPtr<FEOSConnectLoginCredentials>> FAu
 			InAsyncOp.SetError(MoveTemp(Error));
 		}
 
-		return MakeFulfilledPromise<const EOS_Auth_LinkAccountCallbackInfo*>().GetFuture();
+		Promise.EmplaceValue();
 	})
 	.Then([this](TOnlineAsyncOp<FAuthLogin>& InAsyncOp, const EOS_Auth_LinkAccountCallbackInfo* Data) -> TFuture<TSharedPtr<FEOSConnectLoginCredentials>>
 	{
@@ -444,9 +446,9 @@ TOnlineAsyncOpHandle<FAuthLogout> FAuthEOS::Logout(FAuthLogout::Params&& Params)
 		EOS_Auth_DeletePersistentAuthOptions DeletePersistentAuthOptions = {0};
 		DeletePersistentAuthOptions.ApiVersion = EOS_AUTH_DELETEPERSISTENTAUTH_API_LATEST;
 		DeletePersistentAuthOptions.RefreshToken = nullptr; // Is this needed?  Docs say it's needed for consoles
-		NextOp = NextOp.Then([this, DeletePersistentAuthOptions](TOnlineAsyncOp<FAuthLogout>& InAsyncOp)
+		NextOp = NextOp.Then([this, DeletePersistentAuthOptions](TOnlineAsyncOp<FAuthLogout>& InAsyncOp, TPromise<const EOS_Auth_DeletePersistentAuthCallbackInfo*>&& Promise)
 		{
-			return EOS_Async<EOS_Auth_DeletePersistentAuthCallbackInfo>(EOS_Auth_DeletePersistentAuth, AuthHandle, DeletePersistentAuthOptions);
+			EOS_Async(EOS_Auth_DeletePersistentAuth, AuthHandle, DeletePersistentAuthOptions, MoveTemp(Promise));
 		})
 		.Then([](TOnlineAsyncOp<FAuthLogout>& InAsyncOp, const EOS_Auth_DeletePersistentAuthCallbackInfo* Data)
 		{
@@ -456,12 +458,12 @@ TOnlineAsyncOpHandle<FAuthLogout> FAuthEOS::Logout(FAuthLogout::Params&& Params)
 	}
 
 	// Logout
-	NextOp.Then([this, LocalUserEasId](TOnlineAsyncOp<FAuthLogout>& InAsyncOp)
+	NextOp.Then([this, LocalUserEasId](TOnlineAsyncOp<FAuthLogout>& InAsyncOp, TPromise<const EOS_Auth_LogoutCallbackInfo*>&& Promise)
 	{
 		EOS_Auth_LogoutOptions LogoutOptions = { };
 		LogoutOptions.ApiVersion = EOS_AUTH_LOGOUT_API_LATEST;
 		LogoutOptions.LocalUserId = LocalUserEasId;
-		return EOS_Async<EOS_Auth_LogoutCallbackInfo>(EOS_Auth_Logout, AuthHandle, LogoutOptions);
+		EOS_Async(EOS_Auth_Logout, AuthHandle, LogoutOptions, MoveTemp(Promise));
 	})
 	.Then([](TOnlineAsyncOp<FAuthLogout>& InAsyncOp, const EOS_Auth_LogoutCallbackInfo* Data)
 	{
@@ -550,9 +552,6 @@ TFuture<TArray<FOnlineAccountIdHandle>> FAuthEOS::ResolveAccountIds(const FOnlin
 		return MakeFulfilledPromise<TArray<FOnlineAccountIdHandle>>().GetFuture();
 	}
 
-	TPromise<TArray<FOnlineAccountIdHandle>> Promise;
-	TFuture<TArray<FOnlineAccountIdHandle>> Future = Promise.GetFuture();
-
 	TArray<FEpicAccountIdStrBuffer> EpicAccountIdStrsToQuery;
 	EpicAccountIdStrsToQuery.Reserve(MissingEpicAccountIds.Num());
 	for (const EOS_EpicAccountId EpicAccountId : MissingEpicAccountIds)
@@ -570,6 +569,9 @@ TFuture<TArray<FOnlineAccountIdHandle>> FAuthEOS::ResolveAccountIds(const FOnlin
 	TArray<const char*> EpicAccountIdStrPtrs;
 	Algo::Transform(EpicAccountIdStrsToQuery, EpicAccountIdStrPtrs, [](const FEpicAccountIdStrBuffer& Str) { return &Str[0]; });
 
+	TPromise<TArray<FOnlineAccountIdHandle>> Promise;
+	TFuture<TArray<FOnlineAccountIdHandle>> Future = Promise.GetFuture();
+
 	EOS_Connect_QueryExternalAccountMappingsOptions Options = {};
 	Options.ApiVersion = EOS_CONNECT_QUERYEXTERNALACCOUNTMAPPINGS_API_LATEST;
 	Options.LocalUserId = GetProductUserIdChecked(LocalUserId);
@@ -577,8 +579,8 @@ TFuture<TArray<FOnlineAccountIdHandle>> FAuthEOS::ResolveAccountIds(const FOnlin
 	Options.ExternalAccountIds = (const char**)EpicAccountIdStrPtrs.GetData();
 	Options.ExternalAccountIdCount = 1;
 
-	EOS_Async<EOS_Connect_QueryExternalAccountMappingsCallbackInfo>(EOS_Connect_QueryExternalAccountMappings, ConnectHandle, Options)
-		.Next([this, InEpicAccountIds, Promise = MoveTemp(Promise)](const EOS_Connect_QueryExternalAccountMappingsCallbackInfo* Data) mutable
+	EOS_Async(EOS_Connect_QueryExternalAccountMappings, ConnectHandle, Options,
+	[this, InEpicAccountIds, Promise = MoveTemp(Promise)](const EOS_Connect_QueryExternalAccountMappingsCallbackInfo* Data) mutable
 	{
 		TArray<FOnlineAccountIdHandle> AccountIds;
 		AccountIds.Reserve(InEpicAccountIds.Num());
@@ -656,8 +658,9 @@ TFuture<TArray<FOnlineAccountIdHandle>> FAuthEOS::ResolveAccountIds(const FOnlin
 	Options.LocalUserId = GetProductUserIdChecked(LocalUserId);
 	Options.ProductUserIds = MissingProductUserIds.GetData();
 	Options.ProductUserIdCount = MissingProductUserIds.Num();
-	EOS_Async<EOS_Connect_QueryProductUserIdMappingsCallbackInfo>(EOS_Connect_QueryProductUserIdMappings, ConnectHandle, Options)
-		.Next([this, InProductUserIds, Promise = MoveTemp(Promise)](const EOS_Connect_QueryProductUserIdMappingsCallbackInfo* Data) mutable
+
+	EOS_Async(EOS_Connect_QueryProductUserIdMappings, ConnectHandle, Options,
+	[this, InProductUserIds, Promise = MoveTemp(Promise)](const EOS_Connect_QueryProductUserIdMappingsCallbackInfo* Data) mutable
 	{
 		TArray<FOnlineAccountIdHandle> AccountIds;
 		if (Data->ResultCode == EOS_EResult::EOS_Success)
diff --git a/Engine/Plugins/Online/Experimental/OnlineServicesEOS/Source/Private/Online/ExternalUIEOS.cpp b/Engine/Plugins/Online/Experimental/OnlineServicesEOS/Source/Private/Online/ExternalUIEOS.cpp
--- a/Engine/Plugins/Online/Experimental/OnlineServicesEOS/Source/Private/Online/ExternalUIEOS.cpp
+++ b/Engine/Plugins/Online/Experimental/OnlineServicesEOS/Source/Private/Online/ExternalUIEOS.cpp
@@ -96,9 +96,9 @@ TOnlineAsyncOpHandle<FExternalUIShowFriendsUI> FExternalUIEOS::ShowFriendsUI(FEx
 	ShowFriendsOptions.ApiVersion = EOS_UI_SHOWFRIENDS_API_LATEST;
 	ShowFriendsOptions.LocalUserId = LocalUserEasId;
 
-	Op->Then([this, ShowFriendsOptions](TOnlineAsyncOp<FExternalUIShowFriendsUI>& InAsyncOp) mutable
+	Op->Then([this, ShowFriendsOptions](TOnlineAsyncOp<FExternalUIShowFriendsUI>& InAsyncOp, TPromise<const EOS_UI_ShowFriendsCallbackInfo*>&& Promise) mutable
 	{
-		return EOS_Async<EOS_UI_ShowFriendsCallbackInfo>(EOS_UI_ShowFriends, UIHandle, ShowFriendsOptions);
+		EOS_Async(EOS_UI_ShowFriends, UIHandle, ShowFriendsOptions, MoveTemp(Promise));
 	})
 	.Then([this](TOnlineAsyncOp<FExternalUIShowFriendsUI>& InAsyncOp, const EOS_UI_ShowFriendsCallbackInfo* Data)
 	{
diff --git a/Engine/Plugins/Online/Experimental/OnlineServicesEOS/Source/Private/Online/FriendsEOS.cpp b/Engine/Plugins/Online/Experimental/OnlineServicesEOS/Source/Private/Online/FriendsEOS.cpp
--- a/Engine/Plugins/Online/Experimental/OnlineServicesEOS/Source/Private/Online/FriendsEOS.cpp
+++ b/Engine/Plugins/Online/Experimental/OnlineServicesEOS/Source/Private/Online/FriendsEOS.cpp
@@ -79,9 +79,9 @@ TOnlineAsyncOpHandle<FQueryFriends> FFriendsEOS::QueryFriends(FQueryFriends::Par
 		EOS_EpicAccountId LocalUserEasId = GetEpicAccountIdChecked(Params.LocalUserId);
 		QueryFriendsOptions.LocalUserId = LocalUserEasId;
 
-		Op->Then([this, QueryFriendsOptions](TOnlineAsyncOp<FQueryFriends>& InAsyncOp) mutable
+		Op->Then([this, QueryFriendsOptions](TOnlineAsyncOp<FQueryFriends>& InAsyncOp, TPromise<const EOS_Friends_QueryFriendsCallbackInfo*>&& Promise) mutable
 		{
-			return EOS_Async<EOS_Friends_QueryFriendsCallbackInfo>(EOS_Friends_QueryFriends, FriendsHandle, QueryFriendsOptions);
+			EOS_Async(EOS_Friends_QueryFriends, FriendsHandle, QueryFriendsOptions, MoveTemp(Promise));
 		})
 		.Then([this](TOnlineAsyncOp<FQueryFriends>& InAsyncOp, const EOS_Friends_QueryFriendsCallbackInfo* Data) mutable
 		{
@@ -231,7 +231,7 @@ TOnlineAsyncOpHandle<FAddFriend> FFriendsEOS::AddFriend(FAddFriend::Params&& InP
 			SendInviteOptions.ApiVersion = EOS_FRIENDS_SENDINVITE_API_LATEST;
 			SendInviteOptions.LocalUserId = EOSAccountIdFromOnlineServiceAccountId(Params.LocalUserId).GetValue();
 			SendInviteOptions.TargetUserId = EOSAccountIdFromOnlineServiceAccountId(Params.FriendId).GetValue();
-			return EOS_Async<EOS_Friends_SendInviteCallbackInfo>(EOS_Friends_SendInvite, FriendsHandle, SendInviteOptions);
+			return EOS_Async(EOS_Friends_SendInvite, FriendsHandle, SendInviteOptions);
 		})
 		.Then([this](TOnlineAsyncOp<FAddFriend>& InAsyncOp, const EOS_Friends_SendInviteCallbackInfo* Data) mutable
 		{
diff --git a/Engine/Plugins/Online/Experimental/OnlineServicesEOS/Source/Private/Online/PresenceEOS.cpp b/Engine/Plugins/Online/Experimental/OnlineServicesEOS/Source/Private/Online/PresenceEOS.cpp
--- a/Engine/Plugins/Online/Experimental/OnlineServicesEOS/Source/Private/Online/PresenceEOS.cpp
+++ b/Engine/Plugins/Online/Experimental/OnlineServicesEOS/Source/Private/Online/PresenceEOS.cpp
@@ -121,7 +121,7 @@ TOnlineAsyncOpHandle<FQueryPresence> FPresenceEOS::QueryPresence(FQueryPresence:
 			return Op->GetHandle();
 		}
 
-		Op->Then([this](TOnlineAsyncOp<FQueryPresence>& InAsyncOp) mutable
+		Op->Then([this](TOnlineAsyncOp<FQueryPresence>& InAsyncOp, TPromise<const EOS_Presence_QueryPresenceCallbackInfo*>&& Promise)
 			{
 				const FQueryPresence::Params& Params = InAsyncOp.GetParams();
 				EOS_Presence_QueryPresenceOptions QueryPresenceOptions = { };
@@ -129,7 +129,7 @@ TOnlineAsyncOpHandle<FQueryPresence> FPresenceEOS::QueryPresence(FQueryPresence:
 				QueryPresenceOptions.LocalUserId = GetEpicAccountIdChecked(Params.LocalUserId);
 				QueryPresenceOptions.TargetUserId = GetEpicAccountIdChecked(Params.TargetUserId);
 
-				return EOS_Async<EOS_Presence_QueryPresenceCallbackInfo>(EOS_Presence_QueryPresence, PresenceHandle, QueryPresenceOptions);
+				EOS_Async(EOS_Presence_QueryPresence, PresenceHandle, QueryPresenceOptions, MoveTemp(Promise));
 			})
 			.Then([this](TOnlineAsyncOp<FQueryPresence>& InAsyncOp, const EOS_Presence_QueryPresenceCallbackInfo* Data) mutable
 			{
@@ -177,7 +177,7 @@ TOnlineAsyncOpHandle<FUpdatePresence> FPresenceEOS::UpdatePresence(FUpdatePresen
 		return Op->GetHandle();
 	}
 
-	Op->Then([this](TOnlineAsyncOp<FUpdatePresence>& InAsyncOp) mutable
+	Op->Then([this](TOnlineAsyncOp<FUpdatePresence>& InAsyncOp, TPromise<const EOS_Presence_SetPresenceCallbackInfo*>&& Promise) mutable
 	{
 		const FUpdatePresence::Params& Params = InAsyncOp.GetParams();
 		EOS_HPresenceModification ChangeHandle = nullptr;
@@ -198,7 +198,8 @@ TOnlineAsyncOpHandle<FUpdatePresence> FPresenceEOS::UpdatePresence(FUpdatePresen
 			{
 				UE_LOG(LogTemp, Warning, TEXT("UpdatePresence: EOS_PresenceModification_SetStatus failed with result %s"), *LexToString(SetStatusResult));
 				InAsyncOp.SetError(Errors::FromEOSResult(SetStatusResult));
-				return MakeFulfilledPromise<const EOS_Presence_SetPresenceCallbackInfo*>(nullptr).GetFuture();
+				Promise.EmplaceValue();
+				return;
 			}
 
 			// Raw rich text
@@ -214,7 +215,8 @@ TOnlineAsyncOpHandle<FUpdatePresence> FPresenceEOS::UpdatePresence(FUpdatePresen
 			{
 				UE_LOG(LogTemp, Warning, TEXT("UpdatePresence: EOS_PresenceModification_SetRawRichText failed with result %s"), *LexToString(SetRawRichTextResult));
 				InAsyncOp.SetError(Errors::FromEOSResult(SetRawRichTextResult));
-				return MakeFulfilledPromise<const EOS_Presence_SetPresenceCallbackInfo*>(nullptr).GetFuture();
+				Promise.EmplaceValue();
+				return;
 			}
 
 			// EOS needs to be specific on which fields are removed and which aren't, so grab the last presence and check which are deletions
@@ -255,7 +257,8 @@ TOnlineAsyncOpHandle<FUpdatePresence> FPresenceEOS::UpdatePresence(FUpdatePresen
 				{
 					UE_LOG(LogTemp, Warning, TEXT("UpdatePresence: EOS_PresenceModification_DeleteDataOptions failed with result %s"), *LexToString(DeleteDataResult));
 					InAsyncOp.SetError(Errors::FromEOSResult(DeleteDataResult));
-					return MakeFulfilledPromise<const EOS_Presence_SetPresenceCallbackInfo*>(nullptr).GetFuture();
+					Promise.EmplaceValue();
+					return;
 				}
 			}
 
@@ -267,7 +270,8 @@ TOnlineAsyncOpHandle<FUpdatePresence> FPresenceEOS::UpdatePresence(FUpdatePresen
 					// TODO: Move this check higher.  Needs to take into account number of present fields (not just ones updated) and removed fields.
 					UE_LOG(LogTemp, Warning, TEXT("UpdatePresence: Too many presence keys.  %u/%u"), Params.Presence->Properties.Num(), EOS_PRESENCE_DATA_MAX_KEYS);
 					InAsyncOp.SetError(Errors::InvalidParams());
-					return MakeFulfilledPromise<const EOS_Presence_SetPresenceCallbackInfo*>(nullptr).GetFuture();
+					Promise.EmplaceValue();
+					return;
 				}
 				TArray<FTCHARToUTF8, TInlineAllocator<EOS_PRESENCE_DATA_MAX_KEYS * 2>> Utf8Strings;
 				TArray<EOS_Presence_DataRecord, TInlineAllocator<EOS_PRESENCE_DATA_MAX_KEYS>> Records;
@@ -293,7 +297,8 @@ TOnlineAsyncOpHandle<FUpdatePresence> FPresenceEOS::UpdatePresence(FUpdatePresen
 				{
 					UE_LOG(LogTemp, Warning, TEXT("UpdatePresence: EOS_PresenceModification_SetData failed with result %s"), *LexToString(SetDataResult));
 					InAsyncOp.SetError(Errors::FromEOSResult(SetDataResult));
-					return MakeFulfilledPromise<const EOS_Presence_SetPresenceCallbackInfo*>(nullptr).GetFuture();
+					Promise.EmplaceValue();
+					return;
 				}
 			}
 
@@ -301,15 +306,16 @@ TOnlineAsyncOpHandle<FUpdatePresence> FPresenceEOS::UpdatePresence(FUpdatePresen
 			SetPresenceOptions.ApiVersion = EOS_PRESENCE_SETPRESENCE_API_LATEST;
 			SetPresenceOptions.LocalUserId = GetEpicAccountIdChecked(Params.LocalUserId);
 			SetPresenceOptions.PresenceModificationHandle = ChangeHandle;
-			TFuture<const EOS_Presence_SetPresenceCallbackInfo*> EOSAsyncFuture = EOS_Async<EOS_Presence_SetPresenceCallbackInfo>(EOS_Presence_SetPresence, PresenceHandle, SetPresenceOptions);
+
+			EOS_Async(EOS_Presence_SetPresence, PresenceHandle, SetPresenceOptions, MoveTemp(Promise));
 			EOS_PresenceModification_Release(ChangeHandle);
-			return EOSAsyncFuture;
+			return;
 		}
 		else
 		{
 			InAsyncOp.SetError(Errors::FromEOSResult(CreatePresenceModificationResult));
-			return MakeFulfilledPromise<const EOS_Presence_SetPresenceCallbackInfo*>(nullptr).GetFuture();
 		}
+		Promise.EmplaceValue();
 	})
 	.Then([this](TOnlineAsyncOp<FUpdatePresence>& InAsyncOp, const EOS_Presence_SetPresenceCallbackInfo* Data) mutable
 	{
@@ -354,7 +360,7 @@ TOnlineAsyncOpHandle<FPartialUpdatePresence> FPresenceEOS::PartialUpdatePresence
 		}
 
 		// Don't cache anything from Params as they could be modified by another merge in the meanwhile.
-		Op->Then([this](TOnlineAsyncOp<FPartialUpdatePresence>& InAsyncOp) mutable
+		Op->Then([this](TOnlineAsyncOp<FPartialUpdatePresence>& InAsyncOp, TPromise<const EOS_Presence_SetPresenceCallbackInfo*>&& Promise) mutable
 			{
 				const FPartialUpdatePresence::Params& Params = InAsyncOp.GetParams();
 				EOS_HPresenceModification ChangeHandle = nullptr;
@@ -377,7 +383,8 @@ TOnlineAsyncOpHandle<FPartialUpdatePresence> FPresenceEOS::PartialUpdatePresence
 						{
 							UE_LOG(LogTemp, Warning, TEXT("UpdatePresence: EOS_PresenceModification_SetStatus failed with result %s"), *LexToString(SetStatusResult));
 							InAsyncOp.SetError(Errors::FromEOSResult(SetStatusResult));
-							return MakeFulfilledPromise<const EOS_Presence_SetPresenceCallbackInfo*>(nullptr).GetFuture();
+							Promise.EmplaceValue();
+							return;
 						}
 						else
 						{
@@ -400,7 +407,8 @@ TOnlineAsyncOpHandle<FPartialUpdatePresence> FPresenceEOS::PartialUpdatePresence
 						{
 							UE_LOG(LogTemp, Warning, TEXT("UpdatePresence: EOS_PresenceModification_SetRawRichText failed with result %s"), *LexToString(SetRawRichTextResult));
 							InAsyncOp.SetError(Errors::FromEOSResult(SetRawRichTextResult));
-							return MakeFulfilledPromise<const EOS_Presence_SetPresenceCallbackInfo*>(nullptr).GetFuture();
+							Promise.EmplaceValue();
+							return;
 						}
 						else
 						{
@@ -434,7 +442,8 @@ TOnlineAsyncOpHandle<FPartialUpdatePresence> FPresenceEOS::PartialUpdatePresence
 						{
 							UE_LOG(LogTemp, Warning, TEXT("UpdatePresence: EOS_PresenceModification_DeleteDataOptions failed with result %s"), *LexToString(DeleteDataResult));
 							InAsyncOp.SetError(Errors::FromEOSResult(DeleteDataResult));
-							return MakeFulfilledPromise<const EOS_Presence_SetPresenceCallbackInfo*>(nullptr).GetFuture();
+							Promise.EmplaceValue();
+							return;
 						}
 					}
 
@@ -446,7 +455,8 @@ TOnlineAsyncOpHandle<FPartialUpdatePresence> FPresenceEOS::PartialUpdatePresence
 							// TODO: Move this check higher.  Needs to take into account number of present fields (not just ones updated) and removed fields.
 							UE_LOG(LogTemp, Warning, TEXT("UpdatePresence: Too many presence keys.  %u/%u"), Params.Mutations.UpdatedProperties.Num(), EOS_PRESENCE_DATA_MAX_KEYS);
 							InAsyncOp.SetError(Errors::InvalidParams());
-							return MakeFulfilledPromise<const EOS_Presence_SetPresenceCallbackInfo*>(nullptr).GetFuture();
+							Promise.EmplaceValue();
+							return;
 						}
 						TArray<FTCHARToUTF8, TInlineAllocator<EOS_PRESENCE_DATA_MAX_KEYS * 2>> Utf8Strings;
 						TArray<EOS_Presence_DataRecord, TInlineAllocator<EOS_PRESENCE_DATA_MAX_KEYS>> Records;
@@ -472,7 +482,8 @@ TOnlineAsyncOpHandle<FPartialUpdatePresence> FPresenceEOS::PartialUpdatePresence
 						{
 							UE_LOG(LogTemp, Warning, TEXT("UpdatePresence: EOS_PresenceModification_SetData failed with result %s"), *LexToString(SetDataResult));
 							InAsyncOp.SetError(Errors::FromEOSResult(SetDataResult));
-							return MakeFulfilledPromise<const EOS_Presence_SetPresenceCallbackInfo*>(nullptr).GetFuture();
+							Promise.EmplaceValue();
+							return;
 						}
 					}
 
@@ -480,15 +491,15 @@ TOnlineAsyncOpHandle<FPartialUpdatePresence> FPresenceEOS::PartialUpdatePresence
 					SetPresenceOptions.ApiVersion = EOS_PRESENCE_SETPRESENCE_API_LATEST;
 					SetPresenceOptions.LocalUserId = GetEpicAccountIdChecked(Params.LocalUserId);
 					SetPresenceOptions.PresenceModificationHandle = ChangeHandle;
-					TFuture<const EOS_Presence_SetPresenceCallbackInfo*> EOSAsyncFuture = EOS_Async<EOS_Presence_SetPresenceCallbackInfo>(EOS_Presence_SetPresence, PresenceHandle, SetPresenceOptions);
+					EOS_Async(EOS_Presence_SetPresence, PresenceHandle, SetPresenceOptions, MoveTemp(Promise));
 					EOS_PresenceModification_Release(ChangeHandle);
-					return EOSAsyncFuture;
+					return;
 				}
 				else
 				{
 					InAsyncOp.SetError(Errors::FromEOSResult(CreatePresenceModificationResult));
-					return MakeFulfilledPromise<const EOS_Presence_SetPresenceCallbackInfo*>(nullptr).GetFuture();
 				}
+				Promise.EmplaceValue();
 			})
 			.Then([this](TOnlineAsyncOp<FPartialUpdatePresence>& InAsyncOp, const EOS_Presence_SetPresenceCallbackInfo* Data) mutable
 			{
diff --git a/Engine/Plugins/Online/Experimental/OnlineServicesEOSGS/Source/Private/Online/AchievementsEOSGS.cpp b/Engine/Plugins/Online/Experimental/OnlineServicesEOSGS/Source/Private/Online/AchievementsEOSGS.cpp
--- a/Engine/Plugins/Online/Experimental/OnlineServicesEOSGS/Source/Private/Online/AchievementsEOSGS.cpp
+++ b/Engine/Plugins/Online/Experimental/OnlineServicesEOSGS/Source/Private/Online/AchievementsEOSGS.cpp
@@ -68,14 +68,14 @@ TOnlineAsyncOpHandle<FQueryAchievementDefinitions> FAchievementsEOSGS::QueryAchi
 
 	if (!AchievementDefinitions.IsSet())
 	{
-		Op->Then([this](TOnlineAsyncOp<FQueryAchievementDefinitions>& InAsyncOp)
+		Op->Then([this](TOnlineAsyncOp<FQueryAchievementDefinitions>& InAsyncOp, TPromise<const EOS_Achievements_OnQueryDefinitionsCompleteCallbackInfo*>&& Promise)
 		{
 			EOS_Achievements_QueryDefinitionsOptions Options = {};
 			Options.ApiVersion = EOS_ACHIEVEMENTS_QUERYDEFINITIONS_API_LATEST;
 			static_assert(EOS_ACHIEVEMENTS_QUERYDEFINITIONS_API_LATEST == 3, "EOS_Achievements_QueryDefinitionsOptions updated, check new fields");
 			Options.LocalUserId = GetProductUserIdChecked(InAsyncOp.GetParams().LocalUserId);
 
-			return EOS_Async<EOS_Achievements_OnQueryDefinitionsCompleteCallbackInfo>(EOS_Achievements_QueryDefinitions, AchievementsHandle, Options);
+			EOS_Async(EOS_Achievements_QueryDefinitions, AchievementsHandle, Options, MoveTemp(Promise));
 		})
 		.Then([this](TOnlineAsyncOp<FQueryAchievementDefinitions>& InAsyncOp, const EOS_Achievements_OnQueryDefinitionsCompleteCallbackInfo* Data)
 		{
@@ -216,7 +216,7 @@ TOnlineAsyncOpHandle<FQueryAchievementStates> FAchievementsEOSGS::QueryAchieveme
 
 	if (!AchievementStates.Contains(Op->GetParams().LocalUserId))
 	{
-		Op->Then([this](TOnlineAsyncOp<FQueryAchievementStates>& InAsyncOp)
+		Op->Then([this](TOnlineAsyncOp<FQueryAchievementStates>& InAsyncOp, TPromise<const EOS_Achievements_OnQueryPlayerAchievementsCompleteCallbackInfo*>&& Promise)
 		{
 			EOS_Achievements_QueryPlayerAchievementsOptions Options = {};
 			Options.ApiVersion = EOS_ACHIEVEMENTS_QUERYPLAYERACHIEVEMENTS_API_LATEST;
@@ -224,7 +224,7 @@ TOnlineAsyncOpHandle<FQueryAchievementStates> FAchievementsEOSGS::QueryAchieveme
 			Options.LocalUserId = GetProductUserIdChecked(InAsyncOp.GetParams().LocalUserId);
 			Options.TargetUserId = Options.LocalUserId;
 
-			return EOS_Async<EOS_Achievements_OnQueryPlayerAchievementsCompleteCallbackInfo>(EOS_Achievements_QueryPlayerAchievements, AchievementsHandle, Options);
+			EOS_Async(EOS_Achievements_QueryPlayerAchievements, AchievementsHandle, Options, MoveTemp(Promise));
 		})
 		.Then([this](TOnlineAsyncOp<FQueryAchievementStates>& InAsyncOp, const EOS_Achievements_OnQueryPlayerAchievementsCompleteCallbackInfo* Data)
 		{
@@ -352,7 +352,7 @@ TOnlineAsyncOpHandle<FUnlockAchievements> FAchievementsEOSGS::UnlockAchievements
 		}
 	}
 
-	Op->Then([this](TOnlineAsyncOp<FUnlockAchievements>& InAsyncOp)
+	Op->Then([this](TOnlineAsyncOp<FUnlockAchievements>& InAsyncOp, TPromise<const EOS_Achievements_OnUnlockAchievementsCompleteCallbackInfo*>&& Promise)
 	{
 		const TArray<FString>& InAchievementIds = InAsyncOp.GetParams().AchievementIds;
 
@@ -373,7 +373,7 @@ TOnlineAsyncOpHandle<FUnlockAchievements> FAchievementsEOSGS::UnlockAchievements
 		Options.AchievementIds = AchievementIdPtrs.GetData();
 		Options.AchievementsCount = AchievementIdPtrs.Num();
 
-		return EOS_Async<EOS_Achievements_OnUnlockAchievementsCompleteCallbackInfo>(EOS_Achievements_UnlockAchievements, AchievementsHandle, Options);
+		EOS_Async(EOS_Achievements_UnlockAchievements, AchievementsHandle, Options, MoveTemp(Promise));
 	})
 	.Then([this](TOnlineAsyncOp<FUnlockAchievements>& InAsyncOp, const EOS_Achievements_OnUnlockAchievementsCompleteCallbackInfo* Data)
 	{
diff --git a/Engine/Plugins/Online/Experimental/OnlineServicesEOSGS/Source/Private/Online/AuthEOSGS.cpp b/Engine/Plugins/Online/Experimental/OnlineServicesEOSGS/Source/Private/Online/AuthEOSGS.cpp
--- a/Engine/Plugins/Online/Experimental/OnlineServicesEOSGS/Source/Private/Online/AuthEOSGS.cpp
+++ b/Engine/Plugins/Online/Experimental/OnlineServicesEOSGS/Source/Private/Online/AuthEOSGS.cpp
@@ -260,15 +260,15 @@ TOnlineAsyncOpHandle<FAuthLogin> FAuthEOSGS::Login(FAuthLogin::Params&& Params)
 	{
 		return ConnectLoginOp.GetOwningOperation().GetHandle();
 	}
-	ConnectLoginOp.Then([this](TOnlineAsyncOp<FAuthLogin>& InAsyncOp, TSharedPtr<FEOSConnectLoginCredentials>&& InConnectLoginCredentials)
+	ConnectLoginOp.Then([this](TOnlineAsyncOp<FAuthLogin>& InAsyncOp, TSharedPtr<FEOSConnectLoginCredentials>&& InConnectLoginCredentials, TPromise<const EOS_Connect_LoginCallbackInfo*>&& Promise)
 	{
 		EOS_Connect_LoginOptions ConnectLoginOptions = { };
 		ConnectLoginOptions.ApiVersion = EOS_CONNECT_LOGIN_API_LATEST;
 		ConnectLoginOptions.Credentials = InConnectLoginCredentials.Get();
 
-		return EOS_Async<EOS_Connect_LoginCallbackInfo>(EOS_Connect_Login, ConnectHandle, ConnectLoginOptions);
+		EOS_Async(EOS_Connect_Login, ConnectHandle, ConnectLoginOptions, MoveTemp(Promise));
 	})
-	.Then([this](TOnlineAsyncOp<FAuthLogin>& InAsyncOp, const EOS_Connect_LoginCallbackInfo* Data)
+	.Then([this](TOnlineAsyncOp<FAuthLogin>& InAsyncOp, const EOS_Connect_LoginCallbackInfo* Data, TPromise<const EOS_Connect_CreateUserCallbackInfo*>&& Promise)
 	{
 		UE_LOG(LogTemp, Verbose, TEXT("[FAuthEOSGS::Login] EOS_Connect_Login Result: [%s]"), *LexToString(Data->ResultCode));
 
@@ -282,7 +282,8 @@ TOnlineAsyncOpHandle<FAuthLogin> FAuthEOSGS::Login(FAuthLogin::Params&& Params)
 			ConnectCreateUserOptions.ApiVersion = EOS_CONNECT_CREATEUSER_API_LATEST;
 			ConnectCreateUserOptions.ContinuanceToken = Data->ContinuanceToken;
 
-			return EOS_Async<EOS_Connect_CreateUserCallbackInfo>(EOS_Connect_CreateUser, ConnectHandle, ConnectCreateUserOptions);
+			EOS_Async(EOS_Connect_CreateUser, ConnectHandle, ConnectCreateUserOptions, MoveTemp(Promise));
+			return;
 		}
 		else
 		{
@@ -290,7 +291,7 @@ TOnlineAsyncOpHandle<FAuthLogin> FAuthEOSGS::Login(FAuthLogin::Params&& Params)
 			InAsyncOp.SetError(Errors::Unknown()); // TODO
 		}
 
-		return MakeFulfilledPromise<const EOS_Connect_CreateUserCallbackInfo*>().GetFuture();
+		Promise.EmplaceValue();
 	})
 	.Then([this](TOnlineAsyncOp<FAuthLogin>& InAsyncOp, const EOS_Connect_CreateUserCallbackInfo* Data)
 	{
@@ -386,12 +387,12 @@ TOnlineChainableAsyncOp<FAuthLogin, TSharedPtr<FEOSConnectLoginCredentials>> FAu
 		return InAsyncOp.Then(FailureLambda);
 	}
 
-	return InAsyncOp.Then([this, LoginOptions, Credentials](TOnlineAsyncOp<FAuthLogin>& InAsyncOp) mutable
+	return InAsyncOp.Then([this, LoginOptions, Credentials](TOnlineAsyncOp<FAuthLogin>& InAsyncOp, TPromise<const EOS_Auth_LoginCallbackInfo*>&& Promise) mutable
 	{
 		LoginOptions.Credentials = &Credentials;
-		return EOS_Async<EOS_Auth_LoginCallbackInfo>(EOS_Auth_Login, AuthHandle, LoginOptions);
+		EOS_Async(EOS_Auth_Login, AuthHandle, LoginOptions, MoveTemp(Promise));
 	})
-	.Then([this](TOnlineAsyncOp<FAuthLogin>& InAsyncOp, const EOS_Auth_LoginCallbackInfo* Data)
+	.Then([this](TOnlineAsyncOp<FAuthLogin>& InAsyncOp, const EOS_Auth_LoginCallbackInfo* Data, TPromise<const EOS_Auth_LinkAccountCallbackInfo*>&& Promise)
 	{
 		UE_LOG(LogTemp, Verbose, TEXT("[FAuthEOS::Login] EOS_Auth_Login Result: [%s]"), *LexToString(Data->ResultCode));
 
@@ -408,7 +409,8 @@ TOnlineChainableAsyncOp<FAuthLogin, TSharedPtr<FEOSConnectLoginCredentials>> FAu
 			LinkAccountOptions.ApiVersion = EOS_AUTH_LINKACCOUNT_API_LATEST;
 			LinkAccountOptions.ContinuanceToken = Data->ContinuanceToken;
 
-			return EOS_Async<EOS_Auth_LinkAccountCallbackInfo>(EOS_Auth_LinkAccount, AuthHandle, LinkAccountOptions);
+			EOS_Async(EOS_Auth_LinkAccount, AuthHandle, LinkAccountOptions, MoveTemp(Promise));
+			return;
 		}
 		else
 		{
@@ -421,7 +423,7 @@ TOnlineChainableAsyncOp<FAuthLogin, TSharedPtr<FEOSConnectLoginCredentials>> FAu
 			InAsyncOp.SetError(MoveTemp(Error));
 		}
 
-		return MakeFulfilledPromise<const EOS_Auth_LinkAccountCallbackInfo*>().GetFuture();
+		Promise.EmplaceValue();
 	})
 	.Then([this](TOnlineAsyncOp<FAuthLogin>& InAsyncOp, const EOS_Auth_LinkAccountCallbackInfo* Data) -> TFuture<TSharedPtr<FEOSConnectLoginCredentials>>
 	{
diff --git a/Engine/Plugins/Online/Experimental/OnlineServicesEOSGS/Source/Private/Online/LobbiesEOSGS.cpp b/Engine/Plugins/Online/Experimental/OnlineServicesEOSGS/Source/Private/Online/LobbiesEOSGS.cpp
--- a/Engine/Plugins/Online/Experimental/OnlineServicesEOSGS/Source/Private/Online/LobbiesEOSGS.cpp
+++ b/Engine/Plugins/Online/Experimental/OnlineServicesEOSGS/Source/Private/Online/LobbiesEOSGS.cpp
@@ -105,7 +105,7 @@ TOnlineAsyncOpHandle<FCreateLobby> FLobbiesEOSGS::CreateLobby(FCreateLobby::Para
 		return Future;
 	};
 
-	Op->Then([this](TOnlineAsyncOp<FCreateLobby>& InAsyncOp)
+	Op->Then([this](TOnlineAsyncOp<FCreateLobby>& InAsyncOp, TPromise<const EOS_Lobby_CreateLobbyCallbackInfo*>&& Promise)
 	{
 		// Step 1: Call create lobby.
 
@@ -126,7 +126,7 @@ TOnlineAsyncOpHandle<FCreateLobby> FLobbiesEOSGS::CreateLobby(FCreateLobby::Para
 		CreateLobbyOptions.bDisableHostMigration = false; // todo: handle
 		CreateLobbyOptions.bEnableRTCRoom = false; // todo: handle
 
-		return EOS_Async<EOS_Lobby_CreateLobbyCallbackInfo>(EOS_Lobby_CreateLobby, LobbyPrerequisites->LobbyInterfaceHandle, CreateLobbyOptions); 
+		EOS_Async(EOS_Lobby_CreateLobby, LobbyPrerequisites->LobbyInterfaceHandle, CreateLobbyOptions, MoveTemp(Promise)); 
 	})
 	.Then([this](TOnlineAsyncOp<FCreateLobby>& InAsyncOp, const EOS_Lobby_CreateLobbyCallbackInfo* Data)
 	{
@@ -1264,7 +1264,7 @@ TOnlineAsyncOpHandle<FLobbiesEOSGS::FJoinLobbyMemberImpl> FLobbiesEOSGS::JoinLob
 		return Op->GetHandle();
 	}
 
-	Op->Then([this, LobbyDetails, LobbyData = Params.LobbyData](TOnlineAsyncOp<FJoinLobbyMemberImpl>& InAsyncOp)
+	Op->Then([this, LobbyDetails, LobbyData = Params.LobbyData](TOnlineAsyncOp<FJoinLobbyMemberImpl>& InAsyncOp, TPromise<const EOS_Lobby_JoinLobbyCallbackInfo*>&& Promise)
 	{
 		const FLobbiesEOSGS::FJoinLobbyMemberImpl::Params& Params = InAsyncOp.GetParams();
 
@@ -1274,7 +1274,7 @@ TOnlineAsyncOpHandle<FLobbiesEOSGS::FJoinLobbyMemberImpl> FLobbiesEOSGS::JoinLob
 		JoinLobbyOptions.LocalUserId = GetProductUserIdChecked(Params.LocalUserId);
 		JoinLobbyOptions.bPresenceEnabled = false;
 		JoinLobbyOptions.LocalRTCOptions = nullptr;
-		return EOS_Async<EOS_Lobby_JoinLobbyCallbackInfo>(EOS_Lobby_JoinLobby, LobbyPrerequisites->LobbyInterfaceHandle, JoinLobbyOptions);
+		EOS_Async(EOS_Lobby_JoinLobby, LobbyPrerequisites->LobbyInterfaceHandle, JoinLobbyOptions, MoveTemp(Promise));
 	})
 	.Then([this](TOnlineAsyncOp<FJoinLobbyMemberImpl>& InAsyncOp, const EOS_Lobby_JoinLobbyCallbackInfo* Data)
 	{
@@ -1358,10 +1358,9 @@ TFuture<TDefaultErrorResult<FLobbiesEOSGS::FLeaveLobbyImpl>> FLobbiesEOSGS::Leav
 	TPromise<TDefaultErrorResult<FLeaveLobbyImpl>> Promise;
 	TFuture<TDefaultErrorResult<FLeaveLobbyImpl>> Future = Promise.GetFuture();
 
-	EOS_Async<EOS_Lobby_LeaveLobbyCallbackInfo>(EOS_Lobby_LeaveLobby, LobbyPrerequisites->LobbyInterfaceHandle, LeaveLobbyOptions)
-	.Then([Promise = MoveTemp(Promise)](TFuture<const EOS_Lobby_LeaveLobbyCallbackInfo*>&& Future) mutable
+	EOS_Async(EOS_Lobby_LeaveLobby, LobbyPrerequisites->LobbyInterfaceHandle, LeaveLobbyOptions,
+	[Promise = MoveTemp(Promise)](const EOS_Lobby_LeaveLobbyCallbackInfo* Result) mutable
 	{
-		const EOS_Lobby_LeaveLobbyCallbackInfo* Result = Future.Get();
 		if (Result->ResultCode != EOS_EResult::EOS_Success)
 		{
 			// Todo: Errors
@@ -1396,10 +1395,9 @@ TFuture<TDefaultErrorResult<FLobbiesEOSGS::FDestroyLobbyImpl>> FLobbiesEOSGS::De
 	TPromise<TDefaultErrorResult<FDestroyLobbyImpl>> Promise;
 	TFuture<TDefaultErrorResult<FDestroyLobbyImpl>> Future = Promise.GetFuture();
 
-	EOS_Async<EOS_Lobby_DestroyLobbyCallbackInfo>(EOS_Lobby_DestroyLobby, LobbyPrerequisites->LobbyInterfaceHandle, DestroyLobbyOptions)
-	.Then([Promise = MoveTemp(Promise)](TFuture<const EOS_Lobby_DestroyLobbyCallbackInfo*>&& Future) mutable
+	EOS_Async(EOS_Lobby_DestroyLobby, LobbyPrerequisites->LobbyInterfaceHandle, DestroyLobbyOptions,
+	[Promise = MoveTemp(Promise)](const EOS_Lobby_DestroyLobbyCallbackInfo* Result) mutable
 	{
-		const EOS_Lobby_DestroyLobbyCallbackInfo* Result = Future.Get();
 		if (Result->ResultCode != EOS_EResult::EOS_Success)
 		{
 			// Todo: Errors
@@ -1435,10 +1433,9 @@ TFuture<TDefaultErrorResult<FLobbiesEOSGS::FInviteLobbyMemberImpl>> FLobbiesEOSG
 	TPromise<TDefaultErrorResult<FInviteLobbyMemberImpl>> Promise;
 	TFuture<TDefaultErrorResult<FInviteLobbyMemberImpl>> Future = Promise.GetFuture();
 
-	EOS_Async<EOS_Lobby_SendInviteCallbackInfo>(EOS_Lobby_SendInvite, LobbyPrerequisites->LobbyInterfaceHandle, SendInviteOptions)
-	.Then([Promise = MoveTemp(Promise)](TFuture<const EOS_Lobby_SendInviteCallbackInfo*>&& Future) mutable
+	EOS_Async(EOS_Lobby_SendInvite, LobbyPrerequisites->LobbyInterfaceHandle, SendInviteOptions,
+	[Promise = MoveTemp(Promise)](const EOS_Lobby_SendInviteCallbackInfo* Result) mutable
 	{
-		const EOS_Lobby_SendInviteCallbackInfo* Result = Future.Get();
 		if (Result->ResultCode != EOS_EResult::EOS_Success)
 		{
 			// Todo: Errors
@@ -1476,13 +1473,12 @@ TFuture<TDefaultErrorResult<FLobbiesEOSGS::FDeclineLobbyInvitationImpl>> FLobbie
 	TPromise<TDefaultErrorResult<FDeclineLobbyInvitationImpl>> Promise;
 	TFuture<TDefaultErrorResult<FDeclineLobbyInvitationImpl>> Future = Promise.GetFuture();
 
-	EOS_Async<EOS_Lobby_RejectInviteCallbackInfo>(EOS_Lobby_RejectInvite, LobbyPrerequisites->LobbyInterfaceHandle, RejectInviteOptions)
-	.Then([this, Promise = MoveTemp(Promise), InviteData](TFuture<const EOS_Lobby_RejectInviteCallbackInfo*>&& Future) mutable
+	EOS_Async(EOS_Lobby_RejectInvite, LobbyPrerequisites->LobbyInterfaceHandle, RejectInviteOptions,
+	[this, Promise = MoveTemp(Promise), InviteData](const EOS_Lobby_RejectInviteCallbackInfo* Result) mutable
 	{
 		// Remove active invitation.
 		RemoveActiveInvite(InviteData.ToSharedRef());
 
-		const EOS_Lobby_RejectInviteCallbackInfo* Result = Future.Get();
 		if (Result->ResultCode != EOS_EResult::EOS_Success)
 		{
 			// Todo: Errors
@@ -1520,10 +1516,9 @@ TFuture<TDefaultErrorResult<FLobbiesEOSGS::FKickLobbyMemberImpl>> FLobbiesEOSGS:
 	TPromise<TDefaultErrorResult<FKickLobbyMemberImpl>> Promise;
 	TFuture<TDefaultErrorResult<FKickLobbyMemberImpl>> Future = Promise.GetFuture();
 
-	EOS_Async<EOS_Lobby_KickMemberCallbackInfo>(EOS_Lobby_KickMember, LobbyPrerequisites->LobbyInterfaceHandle, KickMemberOptions)
-	.Then([Promise = MoveTemp(Promise)](TFuture<const EOS_Lobby_KickMemberCallbackInfo*>&& Future) mutable
+	EOS_Async(EOS_Lobby_KickMember, LobbyPrerequisites->LobbyInterfaceHandle, KickMemberOptions,
+	[Promise = MoveTemp(Promise)](const EOS_Lobby_KickMemberCallbackInfo* Result) mutable
 	{
-		const EOS_Lobby_KickMemberCallbackInfo* Result = Future.Get();
 		if (Result->ResultCode != EOS_EResult::EOS_Success)
 		{
 			// Todo: Errors
@@ -1561,10 +1556,9 @@ TFuture<TDefaultErrorResult<FLobbiesEOSGS::FPromoteLobbyMemberImpl>> FLobbiesEOS
 	TPromise<TDefaultErrorResult<FPromoteLobbyMemberImpl>> Promise;
 	TFuture<TDefaultErrorResult<FPromoteLobbyMemberImpl>> Future = Promise.GetFuture();
 
-	EOS_Async<EOS_Lobby_PromoteMemberCallbackInfo>(EOS_Lobby_PromoteMember, LobbyPrerequisites->LobbyInterfaceHandle, PromoteMemberOptions)
-	.Then([Promise = MoveTemp(Promise)](TFuture<const EOS_Lobby_PromoteMemberCallbackInfo*>&& Future) mutable
+	EOS_Async(EOS_Lobby_PromoteMember, LobbyPrerequisites->LobbyInterfaceHandle, PromoteMemberOptions,
+	[Promise = MoveTemp(Promise)](const EOS_Lobby_PromoteMemberCallbackInfo* Result) mutable
 	{
-		const EOS_Lobby_PromoteMemberCallbackInfo* Result = Future.Get();
 		if (Result->ResultCode != EOS_EResult::EOS_Success)
 		{
 			// Todo: Errors
diff --git a/Engine/Plugins/Online/Experimental/OnlineServicesEOSGS/Source/Private/Online/LobbiesEOSGSTypes.cpp b/Engine/Plugins/Online/Experimental/OnlineServicesEOSGS/Source/Private/Online/LobbiesEOSGSTypes.cpp
--- a/Engine/Plugins/Online/Experimental/OnlineServicesEOSGS/Source/Private/Online/LobbiesEOSGSTypes.cpp
+++ b/Engine/Plugins/Online/Experimental/OnlineServicesEOSGS/Source/Private/Online/LobbiesEOSGSTypes.cpp
@@ -406,7 +406,7 @@ TFuture<EOS_EResult> FLobbyDetailsEOS::ApplyLobbyDataUpdateFromLocalChanges(
 	FOnlineAccountIdHandle LocalUserId,
 	const FClientLobbyDataChanges& Changes) const
 {
-	EOS_HLobbyModification LobbyModificationHandle = {};
+	EOS_HLobbyModification LobbyModificationHandle = nullptr;
 
 	ON_SCOPE_EXIT
 	{
@@ -482,12 +482,11 @@ TFuture<EOS_EResult> FLobbyDetailsEOS::ApplyLobbyDataUpdateFromLocalChanges(
 	UpdateLobbyOptions.ApiVersion = EOS_LOBBY_UPDATELOBBY_API_LATEST;
 	UpdateLobbyOptions.LobbyModificationHandle = LobbyModificationHandle;
 
-	EOS_Async<EOS_Lobby_UpdateLobbyCallbackInfo>(EOS_Lobby_UpdateLobby, Prerequisites->LobbyInterfaceHandle, UpdateLobbyOptions)
-	.Then([Promise = MoveTemp(Promise)](TFuture<const EOS_Lobby_UpdateLobbyCallbackInfo*> Future) mutable
+	EOS_Async(EOS_Lobby_UpdateLobby, Prerequisites->LobbyInterfaceHandle, UpdateLobbyOptions,
+	[Promise = MoveTemp(Promise)](const EOS_Lobby_UpdateLobbyCallbackInfo* CallbackInfo) mutable
 	{
-		Promise.EmplaceValue(Future.Get()->ResultCode);
+		Promise.EmplaceValue(CallbackInfo->ResultCode);
 	});
-
 	return Future;
 }
 
@@ -556,10 +555,10 @@ TFuture<EOS_EResult> FLobbyDetailsEOS::ApplyLobbyMemberDataUpdateFromLocalChange
 	TPromise<EOS_EResult> Promise;
 	TFuture<EOS_EResult> Future = Promise.GetFuture();
 
-	EOS_Async<EOS_Lobby_UpdateLobbyCallbackInfo>(EOS_Lobby_UpdateLobby, Prerequisites->LobbyInterfaceHandle, UpdateLobbyOptions)
-	.Then([Promise = MoveTemp(Promise)](TFuture<const EOS_Lobby_UpdateLobbyCallbackInfo*> Future) mutable
+	EOS_Async(EOS_Lobby_UpdateLobby, Prerequisites->LobbyInterfaceHandle, UpdateLobbyOptions,
+	[Promise = MoveTemp(Promise)](const EOS_Lobby_UpdateLobbyCallbackInfo* CallbackInfo) mutable
 	{
-		Promise.EmplaceValue(Future.Get()->ResultCode);
+		Promise.EmplaceValue(CallbackInfo->ResultCode);
 	});
 
 	return Future;
@@ -946,14 +945,14 @@ TFuture<TDefaultErrorResultInternal<TSharedRef<FLobbySearchEOS>>> FLobbySearchEO
 	FindOptions.ApiVersion = EOS_LOBBYSEARCH_FIND_API_LATEST;
 	FindOptions.LocalUserId = GetProductUserIdChecked(Params.LocalUserId);
 
-	EOS_Async<EOS_LobbySearch_FindCallbackInfo>(EOS_LobbySearch_Find, SearchHandle->Get(), FindOptions)
-	.Then([Promise = MoveTemp(Promise), Prerequisites, LobbyRegistry, LocalUserId = Params.LocalUserId, SearchHandle]
-	(TFuture<const EOS_LobbySearch_FindCallbackInfo*>&& Future) mutable
+	EOS_Async(EOS_LobbySearch_Find, SearchHandle->Get(), FindOptions,
+	[Promise = MoveTemp(Promise), Prerequisites, LobbyRegistry, LocalUserId = Params.LocalUserId, SearchHandle]
+	(const EOS_LobbySearch_FindCallbackInfo* CallbackInfo) mutable
 	{
-		if (Future.Get()->ResultCode != EOS_EResult::EOS_Success)
+		if (CallbackInfo->ResultCode != EOS_EResult::EOS_Success)
 		{
 			// todo: errors
-			Promise.EmplaceValue(FromEOSError(Future.Get()->ResultCode));
+			Promise.EmplaceValue(FromEOSError(CallbackInfo->ResultCode));
 			return;
 		}
 
diff --git a/Engine/Plugins/Online/Experimental/OnlineServicesEOSGS/Source/Private/Online/OnlineServicesEOSGS.cpp b/Engine/Plugins/Online/Experimental/OnlineServicesEOSGS/Source/Private/Online/OnlineServicesEOSGS.cpp
--- a/Engine/Plugins/Online/Experimental/OnlineServicesEOSGS/Source/Private/Online/OnlineServicesEOSGS.cpp
+++ b/Engine/Plugins/Online/Experimental/OnlineServicesEOSGS/Source/Private/Online/OnlineServicesEOSGS.cpp
@@ -10,6 +10,7 @@
 #include "Online/OnlineIdEOSGS.h"
 #include "Online/OnlineServicesEOSGSTypes.h"
 #include "Online/OnlineServicesEOSGSPlatformFactory.h"
+#include "IEOSSDKManager.h"
 
 #if WITH_ENGINE
 #include "InternetAddrEOS.h"
diff --git a/Engine/Plugins/Online/Experimental/OnlineServicesEOSGS/Source/Private/Online/SessionsEOSGS.cpp b/Engine/Plugins/Online/Experimental/OnlineServicesEOSGS/Source/Private/Online/SessionsEOSGS.cpp
--- a/Engine/Plugins/Online/Experimental/OnlineServicesEOSGS/Source/Private/Online/SessionsEOSGS.cpp
+++ b/Engine/Plugins/Online/Experimental/OnlineServicesEOSGS/Source/Private/Online/SessionsEOSGS.cpp
@@ -741,10 +741,9 @@ TFuture<TDefaultErrorResult<FUpdateSessionImpl>> FSessionsEOSGS::UpdateSessionIm
 	TPromise<TDefaultErrorResult<FUpdateSessionImpl>> Promise;
 	TFuture<TDefaultErrorResult<FUpdateSessionImpl>> Future = Promise.GetFuture();
 
-	EOS_Async<EOS_Sessions_UpdateSessionCallbackInfo>(EOS_Sessions_UpdateSession, SessionsHandle, UpdateSessionOptions)
-		.Then([this, Promise = MoveTemp(Promise)](TFuture<const EOS_Sessions_UpdateSessionCallbackInfo*>&& Future) mutable
+	EOS_Async(EOS_Sessions_UpdateSession, SessionsHandle, UpdateSessionOptions,
+	[this, Promise = MoveTemp(Promise)](const EOS_Sessions_UpdateSessionCallbackInfo* Result) mutable
 	{
-		const EOS_Sessions_UpdateSessionCallbackInfo* Result = Future.Get();
 		if (Result->ResultCode != EOS_EResult::EOS_Success && Result->ResultCode != EOS_EResult::EOS_Sessions_OutOfSync)
 		{
 			Promise.EmplaceValue(FromEOSError(Result->ResultCode));
@@ -770,7 +769,7 @@ TOnlineAsyncOpHandle<FLeaveSession> FSessionsEOSGS::LeaveSession(FLeaveSession::
 {
 	TOnlineAsyncOpRef<FLeaveSession> Op = GetOp<FLeaveSession>(MoveTemp(Params));
 
-	Op->Then([this](TOnlineAsyncOp<FLeaveSession>& Op)
+	Op->Then([this](TOnlineAsyncOp<FLeaveSession>& Op, TPromise<const EOS_Sessions_DestroySessionCallbackInfo*>&& Promise)
 	{
 		const FLeaveSession::Params& OpParams = Op.GetParams();
 
@@ -785,14 +784,14 @@ TOnlineAsyncOpHandle<FLeaveSession> FSessionsEOSGS::LeaveSession(FLeaveSession::
 			FTCHARToUTF8 SessionNameUtf8(OpParams.SessionName.ToString());
 			DestroySessionOptions.SessionName = SessionNameUtf8.Get();
 
-			return EOS_Async<EOS_Sessions_DestroySessionCallbackInfo>(EOS_Sessions_DestroySession, SessionsHandle, DestroySessionOptions);
+			EOS_Async(EOS_Sessions_DestroySession, SessionsHandle, DestroySessionOptions, MoveTemp(Promise));
+			return;
 		}
 		else
 		{
 			Op.SetError(Errors::NotFound());
-
-			return MakeFulfilledPromise<const EOS_Sessions_DestroySessionCallbackInfo*>().GetFuture();
 		}
+		Promise.EmplaceValue();
 	})
 	.Then([this](TOnlineAsyncOp<FLeaveSession>& Op, const EOS_Sessions_DestroySessionCallbackInfo* Result)
 	{
@@ -911,7 +910,7 @@ TOnlineAsyncOpHandle<FFindSessions> FSessionsEOSGS::FindSessions(FFindSessions::
 	}
 
 	// We start the find operation
-	Op->Then([this](TOnlineAsyncOp<FFindSessions>& Op)
+	Op->Then([this](TOnlineAsyncOp<FFindSessions>& Op, TPromise<const EOS_SessionSearch_FindCallbackInfo*>&& Promise)
 	{
 		if (!CurrentSessionSearch.IsValid())
 		{
@@ -933,7 +932,8 @@ TOnlineAsyncOpHandle<FFindSessions> FSessionsEOSGS::FindSessions(FFindSessions::
 
 				Op.SetError(FromEOSError(ResultCode));
 
-				return MakeFulfilledPromise<const EOS_SessionSearch_FindCallbackInfo*>().GetFuture();
+				Promise.EmplaceValue();
+				return;
 			}
 
 			CurrentSessionSearchHandleEOSGS = MakeShareable(new FSessionSearchHandleEOSGS(SearchHandle));
@@ -946,14 +946,14 @@ TOnlineAsyncOpHandle<FFindSessions> FSessionsEOSGS::FindSessions(FFindSessions::
 			static_assert(EOS_SESSIONSEARCH_FIND_API_LATEST == 2, "EOS_SessionSearch_FindOptions updated, check new fields");
 			FindOptions.LocalUserId = GetProductUserIdChecked(OpParams.LocalUserId);
 
-			return EOS_Async<EOS_SessionSearch_FindCallbackInfo>(EOS_SessionSearch_Find, CurrentSessionSearchHandleEOSGS->SearchHandle, FindOptions);
+			EOS_Async(EOS_SessionSearch_Find, CurrentSessionSearchHandleEOSGS->SearchHandle, FindOptions, MoveTemp(Promise));
+			return;
 		}
 		else
 		{
 			Op.SetError(Errors::AlreadyPending());
-
-			return MakeFulfilledPromise<const EOS_SessionSearch_FindCallbackInfo*>().GetFuture();
 		}
+		Promise.EmplaceValue();
 	})
 	.Then([this](TOnlineAsyncOp<FFindSessions>& Op, const EOS_SessionSearch_FindCallbackInfo*&& FindCallbackInfoResult)
 	{
@@ -1033,7 +1033,7 @@ TOnlineAsyncOpHandle<FJoinSession> FSessionsEOSGS::JoinSession(FJoinSession::Par
 	// TODO: Routine JoinSession checks
 
 	// We start the join operation
-	Op->Then([this](TOnlineAsyncOp<FJoinSession>& Op)
+	Op->Then([this](TOnlineAsyncOp<FJoinSession>& Op, TPromise<const EOS_Sessions_JoinSessionCallbackInfo*>&& Promise)
 	{
 		const FJoinSession::Params& OpParams = Op.GetParams();
 
@@ -1063,7 +1063,7 @@ TOnlineAsyncOpHandle<FJoinSession> FSessionsEOSGS::JoinSession(FJoinSession::Par
 		const FSessionEOSGS& SessionEOSGS = FSessionEOSGS::Cast(*OpParams.Session);
 		JoinSessionOptions.SessionHandle = SessionEOSGS.SessionDetailsHandle->SessionDetailsHandle;
 
-		return EOS_Async<EOS_Sessions_JoinSessionCallbackInfo>(EOS_Sessions_JoinSession, SessionsHandle, JoinSessionOptions);
+		EOS_Async(EOS_Sessions_JoinSession, SessionsHandle, JoinSessionOptions, MoveTemp(Promise));
 	})
 	.Then([this](TOnlineAsyncOp<FJoinSession>& Op, const EOS_Sessions_JoinSessionCallbackInfo*&& Result)
 	{
diff --git a/Engine/Plugins/Online/Experimental/OnlineServicesEOSGS/Source/Private/Online/StatsEOSGS.cpp b/Engine/Plugins/Online/Experimental/OnlineServicesEOSGS/Source/Private/Online/StatsEOSGS.cpp
--- a/Engine/Plugins/Online/Experimental/OnlineServicesEOSGS/Source/Private/Online/StatsEOSGS.cpp
+++ b/Engine/Plugins/Online/Experimental/OnlineServicesEOSGS/Source/Private/Online/StatsEOSGS.cpp
@@ -56,7 +56,7 @@ TOnlineAsyncOpHandle<FUpdateStats> FStatsEOSGS::UpdateStats(FUpdateStats::Params
 
 	for (const FUserStats& UpdateUserStats : Op->GetParams().UpdateUsersStats)
 	{
-		Op->Then([this, &UpdateUserStats](TOnlineAsyncOp<FUpdateStats>& InAsyncOp)
+		Op->Then([this, &UpdateUserStats](TOnlineAsyncOp<FUpdateStats>& InAsyncOp, TPromise<const EOS_Stats_IngestStatCompleteCallbackInfo*>&& Promise)
 		{
 			TArray<EOS_Stats_IngestData> EOSData;
 			TArray<FStatNameRawBufferEOS> EOSStatNames;
@@ -87,7 +87,7 @@ TOnlineAsyncOpHandle<FUpdateStats> FStatsEOSGS::UpdateStats(FUpdateStats::Params
 			Options.Stats = EOSData.GetData();
 			Options.StatsCount = EOSData.Num();
 
-			return EOS_Async<EOS_Stats_IngestStatCompleteCallbackInfo>(EOS_Stats_IngestStat, StatsHandle, Options);
+			EOS_Async(EOS_Stats_IngestStat, StatsHandle, Options, MoveTemp(Promise));
 		})
 		.Then([this](TOnlineAsyncOp<FUpdateStats>& InAsyncOp, const EOS_Stats_IngestStatCompleteCallbackInfo* Data)
 		{
@@ -113,7 +113,7 @@ TOnlineAsyncOpHandle<FUpdateStats> FStatsEOSGS::UpdateStats(FUpdateStats::Params
 namespace Private
 {
 
-TFuture<const EOS_Stats_OnQueryStatsCompleteCallbackInfo*> QueryStatsEOS(EOS_HStats StatsHandle, const FOnlineAccountIdHandle& LocalUserId, const FOnlineAccountIdHandle& TargetUserId, const TArray<FString>& StatNames)
+void QueryStatsEOS(EOS_HStats StatsHandle, const FOnlineAccountIdHandle& LocalUserId, const FOnlineAccountIdHandle& TargetUserId, const TArray<FString>& StatNames, TPromise<const EOS_Stats_OnQueryStatsCompleteCallbackInfo*>&& Promise)
 {
 	EOS_Stats_QueryStatsOptions Options;
 	Options.ApiVersion = EOS_STATS_QUERYSTATS_API_LATEST;
@@ -136,7 +136,7 @@ TFuture<const EOS_Stats_OnQueryStatsCompleteCallbackInfo*> QueryStatsEOS(EOS_HSt
 	Options.LocalUserId = GetProductUserIdChecked(LocalUserId);
 	Options.TargetUserId = GetProductUserIdChecked(TargetUserId);
 
-	return EOS_Async<EOS_Stats_OnQueryStatsCompleteCallbackInfo>(EOS_Stats_QueryStats, StatsHandle, Options);
+	EOS_Async(EOS_Stats_QueryStats, StatsHandle, Options, MoveTemp(Promise));
 }
 
 void ReadStatsFromEOSResult(EOS_HStats StatsHandle, const EOS_Stats_OnQueryStatsCompleteCallbackInfo* Data, const TArray<FString>& StatNames, TMap<FString, FStatValue>& OutStats)
@@ -179,9 +179,9 @@ TOnlineAsyncOpHandle<FQueryStats> FStatsEOSGS::QueryStats(FQueryStats::Params&&
 		return Op->GetHandle();
 	}
 
-	Op->Then([this](TOnlineAsyncOp<FQueryStats>& InAsyncOp)
+	Op->Then([this](TOnlineAsyncOp<FQueryStats>& InAsyncOp, TPromise<const EOS_Stats_OnQueryStatsCompleteCallbackInfo*>&& Promise)
 	{
-		return Private::QueryStatsEOS(StatsHandle, InAsyncOp.GetParams().LocalUserId, InAsyncOp.GetParams().TargetUserId, InAsyncOp.GetParams().StatNames);
+		Private::QueryStatsEOS(StatsHandle, InAsyncOp.GetParams().LocalUserId, InAsyncOp.GetParams().TargetUserId, InAsyncOp.GetParams().StatNames, MoveTemp(Promise));
 	})
 	.Then([this](TOnlineAsyncOp<FQueryStats>& InAsyncOp, const EOS_Stats_OnQueryStatsCompleteCallbackInfo* Data)
 	{
@@ -220,9 +220,9 @@ TOnlineAsyncOpHandle<FBatchQueryStats> FStatsEOSGS::BatchQueryStats(FBatchQueryS
 
 	for (const FOnlineAccountIdHandle& TargetUserId : Op->GetParams().TargetUserIds)
 	{
-		Op->Then([this, &TargetUserId](TOnlineAsyncOp<FBatchQueryStats>& InAsyncOp)
+		Op->Then([this, &TargetUserId](TOnlineAsyncOp<FBatchQueryStats>& InAsyncOp, TPromise<const EOS_Stats_OnQueryStatsCompleteCallbackInfo*>&& Promise)
 		{
-			return Private::QueryStatsEOS(StatsHandle, InAsyncOp.GetParams().LocalUserId, TargetUserId, InAsyncOp.GetParams().StatNames);
+			Private::QueryStatsEOS(StatsHandle, InAsyncOp.GetParams().LocalUserId, TargetUserId, InAsyncOp.GetParams().StatNames, MoveTemp(Promise));
 		})
 		.Then([this](TOnlineAsyncOp<FBatchQueryStats>& InAsyncOp, const EOS_Stats_OnQueryStatsCompleteCallbackInfo* Data)
 		{
diff --git a/Engine/Plugins/Online/Experimental/OnlineServicesEOSGS/Source/Private/Online/TitleFileEOSGS.cpp b/Engine/Plugins/Online/Experimental/OnlineServicesEOSGS/Source/Private/Online/TitleFileEOSGS.cpp
--- a/Engine/Plugins/Online/Experimental/OnlineServicesEOSGS/Source/Private/Online/TitleFileEOSGS.cpp
+++ b/Engine/Plugins/Online/Experimental/OnlineServicesEOSGS/Source/Private/Online/TitleFileEOSGS.cpp
@@ -36,20 +36,22 @@ TOnlineAsyncOpHandle<FTitleFileEnumerateFiles> FTitleFileEOSGS::EnumerateFiles(F
 {
 	TOnlineAsyncOpRef<FTitleFileEnumerateFiles> Op = GetOp<FTitleFileEnumerateFiles>(MoveTemp(InParams));
 
-	Op->Then([this](TOnlineAsyncOp<FTitleFileEnumerateFiles>& InAsyncOp)
+	Op->Then([this](TOnlineAsyncOp<FTitleFileEnumerateFiles>& InAsyncOp, TPromise<const EOS_TitleStorage_QueryFileListCallbackInfo*>&& Promise)
 	{
 		const FTitleFileEnumerateFiles::Params& Params = InAsyncOp.GetParams();
 
 		if (!Services.Get<FAuthEOSGS>()->IsLoggedIn(Params.LocalUserId))
 		{
 			InAsyncOp.SetError(Errors::InvalidUser());
-			return MakeFulfilledPromise<const EOS_TitleStorage_QueryFileListCallbackInfo*>().GetFuture();
+			Promise.EmplaceValue();
+			return;
 		}
 
 		if (Config.SearchTag.IsEmpty())
 		{
 			InAsyncOp.SetError(Errors::NotConfigured());
-			return MakeFulfilledPromise<const EOS_TitleStorage_QueryFileListCallbackInfo *>().GetFuture();
+			Promise.EmplaceValue();
+			return;
 		}
 
 		const FTCHARToUTF8 Utf8SearchTag(*Config.SearchTag);
@@ -62,12 +64,12 @@ TOnlineAsyncOpHandle<FTitleFileEnumerateFiles> FTitleFileEOSGS::EnumerateFiles(F
 		Options.ListOfTagsCount = 1;
 		Options.ListOfTags = &SearchTagPtr;
 
-		return EOS_Async<EOS_TitleStorage_QueryFileListCallbackInfo>(EOS_TitleStorage_QueryFileList, TitleStorageHandle, Options);
+		EOS_Async(EOS_TitleStorage_QueryFileList, TitleStorageHandle, Options, MoveTemp(Promise));
 	})
 	.Then([this](TOnlineAsyncOp<FTitleFileEnumerateFiles>& InAsyncOp, const EOS_TitleStorage_QueryFileListCallbackInfo* Data)
 	{
-		const FTitleFileEnumerateFiles::Params& Params = InAsyncOp.GetParams();
-
+		const FTitleFileEnumerateFiles::Params& Params = InAsyncOp.GetParams(); 
+		
 		if (Data->ResultCode != EOS_EResult::EOS_Success)
 		{
 			UE_LOG(LogTemp, Warning, TEXT("EOS_TitleStorage_QueryFileList failed with result=[%s]"), *LexToString(Data->ResultCode));
@@ -147,14 +149,15 @@ TOnlineAsyncOpHandle<FTitleFileReadFile> FTitleFileEOSGS::ReadFile(FTitleFileRea
 		return Op->GetHandle();
 	}
 
-	Op->Then([this](TOnlineAsyncOp<FTitleFileReadFile>& InAsyncOp)
+	Op->Then([this](TOnlineAsyncOp<FTitleFileReadFile>& InAsyncOp, TPromise<const EOS_TitleStorage_ReadFileCallbackInfo*>&& Promise)
 	{
 		const FTitleFileReadFile::Params& Params = InAsyncOp.GetParams();
 
 		if (!Services.Get<FAuthEOSGS>()->IsLoggedIn(Params.LocalUserId))
 		{
 			InAsyncOp.SetError(Errors::InvalidUser());
-			return MakeFulfilledPromise<const EOS_TitleStorage_ReadFileCallbackInfo*>().GetFuture();
+			Promise.EmplaceValue();
+			return;
 		}
 
 		FTitleFileContentsRef FileContentsRef = MakeShared<FTitleFileContents>();
@@ -177,7 +180,7 @@ TOnlineAsyncOpHandle<FTitleFileReadFile> FTitleFileEOSGS::ReadFile(FTitleFileRea
 
 		const FTCHARToUTF8 Utf8Filename(Params.Filename);
 
-		EOS_TitleStorage_ReadFileOptions  Options = {};
+		EOS_TitleStorage_ReadFileOptions Options = {};
 		Options.ApiVersion = EOS_TITLESTORAGE_READFILEOPTIONS_API_LATEST;
 		static_assert(EOS_TITLESTORAGE_READFILEOPTIONS_API_LATEST == 1, "EOS_TitleStorage_ReadFileOptions updated, check new fields");
 		Options.LocalUserId = GetProductUserIdChecked(Params.LocalUserId);
@@ -186,11 +189,8 @@ TOnlineAsyncOpHandle<FTitleFileReadFile> FTitleFileEOSGS::ReadFile(FTitleFileRea
 		Options.ReadFileDataCallback = ReadFileDataCallback->GetCallbackPtr();
 		Options.FileTransferProgressCallback = &FTitleFileEOSGS::OnFileTransferProgressStatic;
 
-		TPromise<const EOS_TitleStorage_ReadFileCallbackInfo*> Promise;
-		TFuture<const EOS_TitleStorage_ReadFileCallbackInfo*> Future = Promise.GetFuture();
-		const EOS_HTitleStorageFileTransferRequest RequestHandle = EOS_Async<EOS_TitleStorage_ReadFileCallbackInfo>(EOS_TitleStorage_ReadFile, TitleStorageHandle, Options, MoveTemp(Promise));
+		const EOS_HTitleStorageFileTransferRequest RequestHandle = EOS_Async(EOS_TitleStorage_ReadFile, TitleStorageHandle, Options, MoveTemp(Promise));
 		InAsyncOp.Data.Set<EOS_HTitleStorageFileTransferRequest>(RequestHandleKey, RequestHandle);
-		return Future;
 	})
 	.Then([this](TOnlineAsyncOp<FTitleFileReadFile>& InAsyncOp, const EOS_TitleStorage_ReadFileCallbackInfo* Data)
 	{
diff --git a/Engine/Plugins/Online/Experimental/OnlineServicesEOSGS/Source/Public/Online/OnlineServicesEOSGSTypes.h b/Engine/Plugins/Online/Experimental/OnlineServicesEOSGS/Source/Public/Online/OnlineServicesEOSGSTypes.h
--- a/Engine/Plugins/Online/Experimental/OnlineServicesEOSGS/Source/Public/Online/OnlineServicesEOSGSTypes.h
+++ b/Engine/Plugins/Online/Experimental/OnlineServicesEOSGS/Source/Public/Online/OnlineServicesEOSGSTypes.h
@@ -16,6 +16,8 @@ inline FOnlineError FromEOSError(EOS_EResult ResultCode)
 	return FOnlineError(Errors::Unknown());
 }
 
+namespace Private {
+
 /** Class to handle all callbacks generically using a future to forward callback results */
 template<typename CallbackType>
 class TEOSCallback
@@ -54,21 +56,69 @@ private:
 	}
 };
 
-template<typename TEOSResult, typename TEOSHandle, typename TEOSParameters, typename TEOSFn> 
-TFuture<const TEOSResult*> EOS_Async(TEOSFn EOSFn, TEOSHandle EOSHandle, TEOSParameters Parameters)
+template<typename Function> struct FEOSAsyncCallbackTraits;
+template<typename TCallbackData>
+struct FEOSAsyncCallbackTraits<void(*)(const TCallbackData*)>
 {
-	TEOSCallback<TEOSResult>* Callback = new TEOSCallback<TEOSResult>();
-	EOSFn(EOSHandle, &Parameters, Callback, *Callback);
-	return Callback->GetFuture();
-}
+	typedef TCallbackData CallbackDataType;
+};
+
+template<typename Function> struct FEOSAsyncTraits;
+template<typename TEOSHandle, typename TEOSParameters, typename TEOSCallback>
+struct FEOSAsyncTraits<void(*)(TEOSHandle, const TEOSParameters*, void*, const TEOSCallback)>
+{
+	typedef TEOSCallback CallbackType;
+	typedef typename FEOSAsyncCallbackTraits<TEOSCallback>::CallbackDataType CallbackDataType;
+};
 
+/* Private */ }
+
+/**
+ * This method requires a Promise that will be fulfilled by the EOS callback.
+ * This is primarily intended to be used with continuations that take a TPromise. e.g.:
+ *
+ *		InAsyncOp.Then([this, LoginOptions, Credentials](TOnlineAsyncOp<FAuthLogin>& InAsyncOp, TPromise<const EOS_Auth_LoginCallbackInfo*>&& Promise) mutable
+ *		{
+ *			LoginOptions.Credentials = &Credentials;
+ *			EOS_Async(EOS_Auth_Login, AuthHandle, LoginOptions, MoveTemp(Promise));
+ *		})
+ * 
+ * Note that when a continuation takes a Promise parameter, the Future for that Promise is already bound
+ * to the next continuation in the Op chain. Therefore, as soon as the EOS_Auth_Login callback is called,
+ * the promise will be fulfilled, and the next continuation will run immediately. This ensures that even
+ * if the SDK calls the callback immediately (i.e. before EOS_Async returns), the continuation is already
+ * bound, and we can safely consume the CallbackInfo. Note that the CallbackInfo is only valid for the
+ * duration of the callback, so it is not safe to bind a continuation _after_ calling the EOS method.
+ */
 template<typename TEOSResult, typename TEOSHandle, typename TEOSParameters, typename TEOSFn>
 decltype(auto) EOS_Async(TEOSFn EOSFn, TEOSHandle EOSHandle, TEOSParameters Parameters, TPromise<const TEOSResult*>&& Promise)
 {
-	TEOSCallback<TEOSResult>* Callback = new TEOSCallback<TEOSResult>(MoveTemp(Promise));
+	Private::TEOSCallback<TEOSResult>* Callback = new Private::TEOSCallback<TEOSResult>(Forward< TPromise<const TEOSResult*>>(Promise));
 	return EOSFn(EOSHandle, &Parameters, Callback, *Callback);
 }
 
+/**
+ * This method requires a Callback (lambda / function pointer) which will be called by the EOS callback
+ * with the CallbackInfo.
+ */
+template<typename TEOSHandle, typename TEOSParameters, typename TEOSFn, typename TCallback>
+decltype(auto) EOS_Async(TEOSFn EOSFn, TEOSHandle EOSHandle, const TEOSParameters& Parameters, TCallback&& Callback)
+{
+	using CallbackDataType = typename Private::FEOSAsyncTraits<TEOSFn>::CallbackDataType;
+	TUniquePtr<TCallback> CallbackObj = MakeUnique<TCallback>(Forward<TCallback>(Callback));
+	return EOSFn(EOSHandle, &Parameters, CallbackObj.Release(),
+	static_cast<void(EOS_CALL*)(const CallbackDataType* Data)>([](const CallbackDataType* Data)
+	{
+		if (EOS_EResult_IsOperationComplete(Data->ResultCode))
+		{
+			check(IsInGameThread());
+			TUniquePtr<TCallback> CallbackObj(reinterpret_cast<TCallback*>(Data->ClientData));
+			check(CallbackObj);
+			(*CallbackObj)(Data);
+		}
+	}));
+}
+
 class EOSEventRegistration
 {
 public:
