- Fix missing change due to it being in #if 0 region

diff --git a/Engine/Plugins/Online/Experimental/OnlineServicesEOS/Source/Private/Online/FriendsEOS.cpp b/Engine/Plugins/Online/Experimental/OnlineServicesEOS/Source/Private/Online/FriendsEOS.cpp
--- a/Engine/Plugins/Online/Experimental/OnlineServicesEOS/Source/Private/Online/FriendsEOS.cpp
+++ b/Engine/Plugins/Online/Experimental/OnlineServicesEOS/Source/Private/Online/FriendsEOS.cpp
@@ -224,14 +224,14 @@ TOnlineAsyncOpHandle<FAddFriend> FFriendsEOS::AddFriend(FAddFriend::Params&& InP
 			return Op->GetHandle();
 		}
 
-		Op->Then([this](TOnlineAsyncOp<FAddFriend>& InAsyncOp) mutable
+		Op->Then([this](TOnlineAsyncOp<FAddFriend>& InAsyncOp, TPromise<const EOS_Friends_SendInviteCallbackInfo*>&& Promise) mutable
 		{
 			const FAddFriend::Params& Params = InAsyncOp.GetParams();
 			EOS_Friends_SendInviteOptions SendInviteOptions = { };
 			SendInviteOptions.ApiVersion = EOS_FRIENDS_SENDINVITE_API_LATEST;
 			SendInviteOptions.LocalUserId = EOSAccountIdFromOnlineServiceAccountId(Params.LocalUserId).GetValue();
 			SendInviteOptions.TargetUserId = EOSAccountIdFromOnlineServiceAccountId(Params.FriendId).GetValue();
-			return EOS_Async(EOS_Friends_SendInvite, FriendsHandle, SendInviteOptions);
+			EOS_Async(EOS_Friends_SendInvite, FriendsHandle, SendInviteOptions, MoveTemp(Promise));
 		})
 		.Then([this](TOnlineAsyncOp<FAddFriend>& InAsyncOp, const EOS_Friends_SendInviteCallbackInfo* Data) mutable
 		{
