- Fix crash in ResolveAccountIds callback firing after OnlineServicesEOS is shut down.

diff --git a/Engine/Plugins/Online/Experimental/OnlineServices/Source/OnlineServicesCommon/Public/Online/OnlineComponent.h b/Engine/Plugins/Online/Experimental/OnlineServices/Source/OnlineServicesCommon/Public/Online/OnlineComponent.h
--- a/Engine/Plugins/Online/Experimental/OnlineServices/Source/OnlineServicesCommon/Public/Online/OnlineComponent.h
+++ b/Engine/Plugins/Online/Experimental/OnlineServices/Source/OnlineServicesCommon/Public/Online/OnlineComponent.h
@@ -176,6 +176,16 @@ public:
 		return MoveTemp(SharedThis).ToSharedRef();
 	}
 
+	TWeakPtr<ComponentType> AsWeak()
+	{
+		return WeakThis;
+	}
+
+	TWeakPtr<ComponentType const> AsWeak() const
+	{
+		return WeakThis;
+	}
+
 protected:
 	template <class OtherType>
 	static TSharedRef<OtherType> SharedThis(OtherType* ThisPtr)
diff --git a/Engine/Plugins/Online/Experimental/OnlineServicesEOS/Source/Private/Online/AuthEOS.cpp b/Engine/Plugins/Online/Experimental/OnlineServicesEOS/Source/Private/Online/AuthEOS.cpp
--- a/Engine/Plugins/Online/Experimental/OnlineServicesEOS/Source/Private/Online/AuthEOS.cpp
+++ b/Engine/Plugins/Online/Experimental/OnlineServicesEOS/Source/Private/Online/AuthEOS.cpp
@@ -580,35 +580,38 @@ TFuture<TArray<FOnlineAccountIdHandle>> FAuthEOS::ResolveAccountIds(const FOnlin
 	Options.ExternalAccountIdCount = 1;
 
 	EOS_Async(EOS_Connect_QueryExternalAccountMappings, ConnectHandle, Options,
-	[this, InEpicAccountIds, Promise = MoveTemp(Promise)](const EOS_Connect_QueryExternalAccountMappingsCallbackInfo* Data) mutable
+	[this, WeakThis = AsWeak(), InEpicAccountIds, Promise = MoveTemp(Promise)](const EOS_Connect_QueryExternalAccountMappingsCallbackInfo* Data) mutable
 	{
 		TArray<FOnlineAccountIdHandle> AccountIds;
-		AccountIds.Reserve(InEpicAccountIds.Num());
-		if (Data->ResultCode == EOS_EResult::EOS_Success)
+		if (const TSharedPtr<IAuth> StrongThis = WeakThis.Pin())
 		{
-			EOS_Connect_GetExternalAccountMappingsOptions Options = {};
-			Options.ApiVersion = EOS_CONNECT_GETEXTERNALACCOUNTMAPPING_API_LATEST;
-			Options.LocalUserId = Data->LocalUserId;
-			Options.AccountIdType = EOS_EExternalAccountType::EOS_EAT_EPIC;
-
-			for (const EOS_EpicAccountId EpicAccountId : InEpicAccountIds)
+			AccountIds.Reserve(InEpicAccountIds.Num());
+			if (Data->ResultCode == EOS_EResult::EOS_Success)
 			{
-				FOnlineAccountIdHandle AccountId = FindAccountId(EpicAccountId);
-				if (!AccountId.IsValid())
+				EOS_Connect_GetExternalAccountMappingsOptions Options = {};
+				Options.ApiVersion = EOS_CONNECT_GETEXTERNALACCOUNTMAPPING_API_LATEST;
+				Options.LocalUserId = Data->LocalUserId;
+				Options.AccountIdType = EOS_EExternalAccountType::EOS_EAT_EPIC;
+
+				for (const EOS_EpicAccountId EpicAccountId : InEpicAccountIds)
 				{
-					FEpicAccountIdStrBuffer EpicAccountIdStr;
-					int32_t BufferSize = sizeof(EpicAccountIdStr);
-					verify(EOS_EpicAccountId_ToString(EpicAccountId, EpicAccountIdStr, &BufferSize) == EOS_EResult::EOS_Success);
-					Options.TargetExternalUserId = &EpicAccountIdStr[0];
-					const EOS_ProductUserId ProductUserId = EOS_Connect_GetExternalAccountMapping(ConnectHandle, &Options);
-					AccountId = CreateAccountId(EpicAccountId, ProductUserId);
+					FOnlineAccountIdHandle AccountId = FindAccountId(EpicAccountId);
+					if (!AccountId.IsValid())
+					{
+						FEpicAccountIdStrBuffer EpicAccountIdStr;
+						int32_t BufferSize = sizeof(EpicAccountIdStr);
+						verify(EOS_EpicAccountId_ToString(EpicAccountId, EpicAccountIdStr, &BufferSize) == EOS_EResult::EOS_Success);
+						Options.TargetExternalUserId = &EpicAccountIdStr[0];
+						const EOS_ProductUserId ProductUserId = EOS_Connect_GetExternalAccountMapping(ConnectHandle, &Options);
+						AccountId = CreateAccountId(EpicAccountId, ProductUserId);
+					}
+					AccountIds.Emplace(MoveTemp(AccountId));
 				}
-				AccountIds.Emplace(MoveTemp(AccountId));
 			}
-		}
-		else
-		{
-			UE_LOG(LogTemp, Warning, TEXT("ResolveAccountId failed to query external mapping Result=[%s]"), *LexToString(Data->ResultCode));
+			else
+			{
+				UE_LOG(LogTemp, Warning, TEXT("ResolveAccountId failed to query external mapping Result=[%s]"), *LexToString(Data->ResultCode));
+			}
 		}
 		Promise.SetValue(MoveTemp(AccountIds));
 	});
@@ -660,39 +663,42 @@ TFuture<TArray<FOnlineAccountIdHandle>> FAuthEOS::ResolveAccountIds(const FOnlin
 	Options.ProductUserIdCount = MissingProductUserIds.Num();
 
 	EOS_Async(EOS_Connect_QueryProductUserIdMappings, ConnectHandle, Options,
-	[this, InProductUserIds, Promise = MoveTemp(Promise)](const EOS_Connect_QueryProductUserIdMappingsCallbackInfo* Data) mutable
+	[this, WeakThis = AsWeak(), InProductUserIds, Promise = MoveTemp(Promise)](const EOS_Connect_QueryProductUserIdMappingsCallbackInfo* Data) mutable
 	{
 		TArray<FOnlineAccountIdHandle> AccountIds;
-		if (Data->ResultCode == EOS_EResult::EOS_Success)
+		if (const TSharedPtr<IAuth> StrongThis = WeakThis.Pin())
 		{
-			EOS_Connect_GetProductUserIdMappingOptions Options = {};
-			Options.ApiVersion = EOS_CONNECT_GETPRODUCTUSERIDMAPPING_API_LATEST;
-			Options.LocalUserId = Data->LocalUserId;
-			Options.AccountIdType = EOS_EExternalAccountType::EOS_EAT_EPIC;
-
-			for (const EOS_ProductUserId ProductUserId : InProductUserIds)
+			if (Data->ResultCode == EOS_EResult::EOS_Success)
 			{
-				FOnlineAccountIdHandle AccountId = FindAccountId(ProductUserId);
-				if (!AccountId.IsValid())
+				EOS_Connect_GetProductUserIdMappingOptions Options = {};
+				Options.ApiVersion = EOS_CONNECT_GETPRODUCTUSERIDMAPPING_API_LATEST;
+				Options.LocalUserId = Data->LocalUserId;
+				Options.AccountIdType = EOS_EExternalAccountType::EOS_EAT_EPIC;
+
+				for (const EOS_ProductUserId ProductUserId : InProductUserIds)
 				{
-					Options.TargetProductUserId = ProductUserId;
-					FEpicAccountIdStrBuffer EpicAccountIdStr;
-					int32_t BufferLength = sizeof(EpicAccountIdStr);
-					EOS_EpicAccountId EpicAccountId = nullptr;
-					const EOS_EResult Result = EOS_Connect_GetProductUserIdMapping(ConnectHandle, &Options, EpicAccountIdStr, &BufferLength);
-					if (Result == EOS_EResult::EOS_Success)
+					FOnlineAccountIdHandle AccountId = FindAccountId(ProductUserId);
+					if (!AccountId.IsValid())
 					{
-						EpicAccountId = EOS_EpicAccountId_FromString(EpicAccountIdStr);
-						check(EOS_EpicAccountId_IsValid(EpicAccountId));
+						Options.TargetProductUserId = ProductUserId;
+						FEpicAccountIdStrBuffer EpicAccountIdStr;
+						int32_t BufferLength = sizeof(EpicAccountIdStr);
+						EOS_EpicAccountId EpicAccountId = nullptr;
+						const EOS_EResult Result = EOS_Connect_GetProductUserIdMapping(ConnectHandle, &Options, EpicAccountIdStr, &BufferLength);
+						if (Result == EOS_EResult::EOS_Success)
+						{
+							EpicAccountId = EOS_EpicAccountId_FromString(EpicAccountIdStr);
+							check(EOS_EpicAccountId_IsValid(EpicAccountId));
+						}
+						AccountId = CreateAccountId(EpicAccountId, ProductUserId);
 					}
-					AccountId = CreateAccountId(EpicAccountId, ProductUserId);
+					AccountIds.Emplace(MoveTemp(AccountId));
 				}
-				AccountIds.Emplace(MoveTemp(AccountId));
 			}
-		}
-		else
-		{
-			UE_LOG(LogTemp, Warning, TEXT("ResolveAccountId failed to query external mapping Result=[%s]"), *LexToString(Data->ResultCode));
+			else
+			{
+				UE_LOG(LogTemp, Warning, TEXT("ResolveAccountId failed to query external mapping Result=[%s]"), *LexToString(Data->ResultCode));
+			}
 		}
 		Promise.SetValue(MoveTemp(AccountIds));
 	});
