- Implement UserInfo for EOS

diff --git a/Engine/Plugins/Online/Experimental/OnlineServicesEOS/Source/Private/Online/OnlineServicesEOS.cpp b/Engine/Plugins/Online/Experimental/OnlineServicesEOS/Source/Private/Online/OnlineServicesEOS.cpp
--- a/Engine/Plugins/Online/Experimental/OnlineServicesEOS/Source/Private/Online/OnlineServicesEOS.cpp
+++ b/Engine/Plugins/Online/Experimental/OnlineServicesEOS/Source/Private/Online/OnlineServicesEOS.cpp
@@ -3,9 +3,10 @@
 #include "Online/OnlineServicesEOS.h"
 
 #include "Online/AuthEOS.h"
+#include "Online/ExternalUIEOS.h"
 #include "Online/FriendsEOS.h"
 #include "Online/PresenceEOS.h"
-#include "Online/ExternalUIEOS.h"
+#include "Online/UserInfoEOS.h"
 
 namespace UE::Online {
 
@@ -17,9 +18,10 @@ FOnlineServicesEOS::FOnlineServicesEOS(FName InInstanceName)
 void FOnlineServicesEOS::RegisterComponents()
 {
 	Components.Register<FAuthEOS>(*this);
+	Components.Register<FExternalUIEOS>(*this);
 	Components.Register<FFriendsEOS>(*this);
 	Components.Register<FPresenceEOS>(*this);
-	Components.Register<FExternalUIEOS>(*this);
+	Components.Register<FUserInfoEOS>(*this);
 
 	Super::RegisterComponents();
 }
diff --git a/Engine/Plugins/Online/Experimental/OnlineServicesEOS/Source/Private/Online/UserInfoEOS.cpp b/Engine/Plugins/Online/Experimental/OnlineServicesEOS/Source/Private/Online/UserInfoEOS.cpp
new file mode 100644
--- /dev/null
+++ b/Engine/Plugins/Online/Experimental/OnlineServicesEOS/Source/Private/Online/UserInfoEOS.cpp
@@ -0,0 +1,125 @@
+// Copyright Epic Games, Inc. All Rights Reserved.
+
+#include "Online/UserInfoEOS.h"
+#include "Online/OnlineAsyncOp.h"
+#include "Online/OnlineIdEOS.h"
+#include "Online/OnlineServicesEOS.h"
+#include "Online/OnlineServicesEOSTypes.h"
+#include "Online/AuthEOS.h"
+#include "Online/ErrorsEOS.h"
+
+#include "eos_userinfo.h"
+
+namespace UE::Online {
+
+FUserInfoEOS::FUserInfoEOS(FOnlineServicesEOS& InServices)
+	: FUserInfoCommon(InServices)
+{
+}
+
+void FUserInfoEOS::Initialize()
+{
+	Super::Initialize();
+
+	UserInfoHandle = EOS_Platform_GetUserInfoInterface(static_cast<FOnlineServicesEOS&>(GetServices()).GetEOSPlatformHandle());
+	check(UserInfoHandle != nullptr);
+}
+
+TOnlineAsyncOpHandle<FQueryUserInfo> FUserInfoEOS::QueryUserInfo(FQueryUserInfo::Params&& InParams)
+{
+	TOnlineAsyncOpRef<FQueryUserInfo> Op = GetJoinableOp<FQueryUserInfo>(MoveTemp(InParams));
+	const FQueryUserInfo::Params& Params = Op->GetParams();
+
+	if (Params.UserIds.IsEmpty())
+	{
+		Op->SetError(Errors::InvalidParams());
+		return Op->GetHandle();
+	}
+
+	for (const FOnlineAccountIdHandle TargetUserId : Params.UserIds)
+	{
+		Op->Then([this, TargetUserId](TOnlineAsyncOp<FQueryUserInfo>& Op, TPromise<const EOS_UserInfo_QueryUserInfoCallbackInfo*>&& Promise)
+		{
+			const FQueryUserInfo::Params& Params = Op.GetParams();
+
+			if (!Services.Get<FAuthEOS>()->IsLoggedIn(Params.LocalUserId))
+			{
+				Op.SetError(Errors::NotLoggedIn());
+				Promise.EmplaceValue();
+				return;
+			}
+
+			const EOS_EpicAccountId TargetUserEasId = GetEpicAccountId(TargetUserId);
+			if (!EOS_EpicAccountId_IsValid(TargetUserEasId))
+			{
+				Op.SetError(Errors::InvalidParams());
+				Promise.EmplaceValue();
+				return;
+			}
+
+			EOS_UserInfo_QueryUserInfoOptions QueryUserInfoOptions = {};
+			QueryUserInfoOptions.ApiVersion = EOS_USERINFO_QUERYUSERINFO_API_LATEST;
+			static_assert(EOS_USERINFO_QUERYUSERINFO_API_LATEST == 1, "EOS_UserInfo_QueryUserInfoOptions updated, check new fields");
+			QueryUserInfoOptions.LocalUserId = GetEpicAccountIdChecked(Params.LocalUserId);
+			QueryUserInfoOptions.TargetUserId = TargetUserEasId;
+
+			EOS_Async(EOS_UserInfo_QueryUserInfo, UserInfoHandle, QueryUserInfoOptions, MoveTemp(Promise));
+		}).Then([this](TOnlineAsyncOp<FQueryUserInfo>& Op, const EOS_UserInfo_QueryUserInfoCallbackInfo* CallbackInfo) mutable
+		{
+			if (CallbackInfo->ResultCode != EOS_EResult::EOS_Success)
+			{
+				UE_LOG(LogTemp, Warning, TEXT("EOS_UserInfo_QueryUserInfo failed with result=[%s]"), *LexToString(CallbackInfo->ResultCode));
+				Op.SetError(FromEOSError(CallbackInfo->ResultCode));
+			}
+		});
+	}
+
+	Op->Then([](TOnlineAsyncOp<FQueryUserInfo>& Op)
+	{
+		Op.SetResult({});
+	});
+	
+	Op->Enqueue(GetSerialQueue());
+	return Op->GetHandle();
+}
+
+TOnlineResult<FGetUserInfo> FUserInfoEOS::GetUserInfo(FGetUserInfo::Params&& Params)
+{
+	if (!Services.Get<FAuthEOS>()->IsLoggedIn(Params.LocalUserId))
+	{
+		return TOnlineResult<FGetUserInfo>(Errors::NotLoggedIn());
+	}
+
+	const EOS_EpicAccountId TargetUserEasId = GetEpicAccountId(Params.UserId);
+	if (!EOS_EpicAccountId_IsValid(TargetUserEasId))
+	{
+		return TOnlineResult<FGetUserInfo>(Errors::InvalidParams());
+	}
+
+	EOS_UserInfo_CopyUserInfoOptions Options;
+	Options.ApiVersion = EOS_USERINFO_COPYUSERINFO_API_LATEST;
+	static_assert(EOS_USERINFO_COPYUSERINFO_API_LATEST == 3, "EOS_UserInfo_CopyUserInfoOptions updated, check new fields");
+	Options.LocalUserId = GetEpicAccountIdChecked(Params.LocalUserId);
+	Options.TargetUserId = TargetUserEasId;
+
+	EOS_UserInfo* EosUserInfo = nullptr;
+	ON_SCOPE_EXIT
+	{
+		EOS_UserInfo_Release(EosUserInfo);
+	};
+
+	const EOS_EResult EosResult = EOS_UserInfo_CopyUserInfo(UserInfoHandle, &Options, &EosUserInfo);
+	if(EosResult != EOS_EResult::EOS_Success)
+	{
+		UE_LOG(LogTemp, Warning, TEXT("EOS_UserInfo_CopyUserInfo failed with result=[%s]"), *LexToString(EosResult));
+		return TOnlineResult<FGetUserInfo>(FromEOSError(EosResult));
+	}
+
+	TSharedRef<FUserInfo> UserInfo = MakeShared<FUserInfo>();
+	UserInfo->UserId = Params.UserId;
+	UserInfo->DisplayName = UTF8_TO_TCHAR(EosUserInfo->DisplayName);
+
+	return TOnlineResult<FGetUserInfo>({UserInfo});
+}
+
+/* UE::Online */ }
diff --git a/Engine/Plugins/Online/Experimental/OnlineServicesEOS/Source/Public/Online/UserInfoEOS.h b/Engine/Plugins/Online/Experimental/OnlineServicesEOS/Source/Public/Online/UserInfoEOS.h
new file mode 100644
--- /dev/null
+++ b/Engine/Plugins/Online/Experimental/OnlineServicesEOS/Source/Public/Online/UserInfoEOS.h
@@ -0,0 +1,36 @@
+// Copyright Epic Games, Inc. All Rights Reserved.
+
+#pragma once
+
+#include "CoreMinimal.h"
+#include "Online/UserInfoCommon.h"
+
+#if defined(EOS_PLATFORM_BASE_FILE_NAME)
+#include EOS_PLATFORM_BASE_FILE_NAME
+#endif
+#include "eos_userinfo_types.h"
+
+namespace UE::Online {
+
+class FOnlineServicesEOS;
+	
+class FUserInfoEOS : public FUserInfoCommon
+{
+public:
+	using Super = FUserInfoCommon;
+
+	FUserInfoEOS(FOnlineServicesEOS& InServices);
+	virtual ~FUserInfoEOS() = default;
+
+	// IOnlineComponent
+	virtual void Initialize() override;
+
+	// IUserInfo
+	virtual TOnlineAsyncOpHandle<FQueryUserInfo> QueryUserInfo(FQueryUserInfo::Params&& Params) override;
+	virtual TOnlineResult<FGetUserInfo> GetUserInfo(FGetUserInfo::Params&& Params) override;
+
+protected:
+	EOS_HUserInfo UserInfoHandle = nullptr;
+};
+
+/* UE::Online */ }
