Fix stringification in log string

diff --git a/Engine/Plugins/Online/Experimental/OnlineServicesEOS/Source/Private/Online/FriendsEOS.cpp b/Engine/Plugins/Online/Experimental/OnlineServicesEOS/Source/Private/Online/FriendsEOS.cpp
--- a/Engine/Plugins/Online/Experimental/OnlineServicesEOS/Source/Private/Online/FriendsEOS.cpp
+++ b/Engine/Plugins/Online/Experimental/OnlineServicesEOS/Source/Private/Online/FriendsEOS.cpp
@@ -248,7 +248,7 @@ TOnlineAsyncOpHandle<FAddFriend> FFriendsEOS::AddFriend(FAddFriend::Params&& InP
 void FFriendsEOS::OnEOSFriendsUpdate(FOnlineAccountIdHandle LocalUserId, FOnlineAccountIdHandle FriendUserId, EOS_EFriendsStatus PreviousStatus, EOS_EFriendsStatus CurrentStatus)
 {
 	// TODO
-	UE_LOG(LogTemp, Warning, TEXT("OnEOSFriendsUpdate: LocalUserId=[%s] FriendUserId=[%s] PreviousStatus=[%s] CurrentStatus=[%s]"), *ToLogString(LocalUserId), *ToLogString(FriendUserId), *LexToString(PreviousStatus), *LexToString(CurrentStatus));
+	UE_LOG(LogTemp, Warning, TEXT("OnEOSFriendsUpdate: LocalUserId=[%s] FriendUserId=[%s] PreviousStatus=[%s] CurrentStatus=[%s]"), *ToLogString(LocalUserId), *ToLogString(FriendUserId), LexToString(PreviousStatus), LexToString(CurrentStatus));
 
 	TMap<FOnlineAccountIdHandle, TSharedRef<FFriend>>& FriendsList = FriendsLists.FindOrAdd(LocalUserId);
 	bool bAnyChanges = false;
