Fixes EOSShared module not loading up when using OSSv2 in a standalone application.

diff --git a/Engine/Plugins/Online/Experimental/OnlineServicesEOS/Source/Private/OnlineServicesEOSModule.cpp b/Engine/Plugins/Online/Experimental/OnlineServicesEOS/Source/Private/OnlineServicesEOSModule.cpp
--- a/Engine/Plugins/Online/Experimental/OnlineServicesEOS/Source/Private/OnlineServicesEOSModule.cpp
+++ b/Engine/Plugins/Online/Experimental/OnlineServicesEOS/Source/Private/OnlineServicesEOSModule.cpp
@@ -31,6 +31,12 @@ protected:
 
 void FOnlineServicesEOSModule::StartupModule()
 {
+	const FName EOSSharedModuleName = TEXT("EOSShared");
+	if (!FModuleManager::Get().IsModuleLoaded(EOSSharedModuleName))
+	{
+		FModuleManager::Get().LoadModuleChecked(EOSSharedModuleName);
+	}
+
 	// Make this higher priority that EOSGS
 	const int Priority = 1;
 	FOnlineServicesRegistry::Get().RegisterServicesFactory(EOnlineServices::Epic, MakeUnique<FOnlineServicesFactoryEOS>(), Priority);
