Commerce EOS framework setup

diff --git a/Engine/Plugins/Online/Experimental/OnlineServicesEOS/Source/Private/Online/CommerceEOS.cpp b/Engine/Plugins/Online/Experimental/OnlineServicesEOS/Source/Private/Online/CommerceEOS.cpp
new file mode 100644
--- /dev/null
+++ b/Engine/Plugins/Online/Experimental/OnlineServicesEOS/Source/Private/Online/CommerceEOS.cpp
@@ -0,0 +1,91 @@
+// Copyright Epic Games, Inc. All Rights Reserved.
+
+#include "Online/CommerceEOS.h"
+#include "Online/OnlineAsyncOp.h"
+#include "Online/OnlineErrorDefinitions.h"
+#include "Online/OnlineServicesCommon.h"
+
+namespace UE::Online {
+
+void FCommerceEOS::PostInitialize()
+{	
+	Super::PostInitialize();
+}
+
+void FCommerceEOS::PreShutdown()
+{
+	Super::PreShutdown();
+}
+
+TOnlineAsyncOpHandle<FCommerceQueryOffers> FCommerceEOS::QueryOffers(FCommerceQueryOffers::Params&& Params)
+{
+	TOnlineAsyncOpRef<FCommerceQueryOffers> Op = GetOp<FCommerceQueryOffers>(MoveTemp(Params));
+	Op->SetError(Errors::NotImplemented());
+	return Op->GetHandle();
+}
+
+TOnlineAsyncOpHandle<FCommerceQueryOffersById> FCommerceEOS::QueryOffersById(FCommerceQueryOffersById::Params&& Params)
+{
+	TOnlineAsyncOpRef<FCommerceQueryOffersById> Op = GetOp<FCommerceQueryOffersById>(MoveTemp(Params));
+	Op->SetError(Errors::NotImplemented());
+	return Op->GetHandle();
+}
+
+TOnlineResult<FCommerceGetOffers> FCommerceEOS::GetOffers(FCommerceGetOffers::Params&& Params)
+{
+	return TOnlineResult<FCommerceGetOffers>(Errors::NotImplemented());
+}
+
+TOnlineResult<FCommerceGetOffersById> FCommerceEOS::GetOffersById(FCommerceGetOffersById::Params&& Params)
+{
+	return TOnlineResult<FCommerceGetOffersById>(Errors::NotImplemented());
+}
+
+TOnlineAsyncOpHandle<FCommerceShowStoreUI> FCommerceEOS::ShowStoreUI(FCommerceShowStoreUI::Params&& Params)
+{
+	TOnlineAsyncOpRef<FCommerceShowStoreUI> Op = GetOp<FCommerceShowStoreUI>(MoveTemp(Params));
+	Op->SetError(Errors::NotImplemented());
+	return Op->GetHandle();
+}
+
+TOnlineAsyncOpHandle<FCommerceCheckout> FCommerceEOS::Checkout(FCommerceCheckout::Params&& Params)
+{
+	TOnlineAsyncOpRef<FCommerceCheckout> Op = GetOp<FCommerceCheckout>(MoveTemp(Params));
+	Op->SetError(Errors::NotImplemented());
+	return Op->GetHandle();
+}
+
+TOnlineAsyncOpHandle<FCommerceQueryTransactionEntitlements> FCommerceEOS::QueryTransactionEntitlements(FCommerceQueryTransactionEntitlements::Params&& Params)
+{
+	TOnlineAsyncOpRef<FCommerceQueryTransactionEntitlements> Op = GetOp<FCommerceQueryTransactionEntitlements>(MoveTemp(Params));
+	Op->SetError(Errors::NotImplemented());
+	return Op->GetHandle();
+}
+
+TOnlineAsyncOpHandle<FCommerceQueryEntitlements> FCommerceEOS::QueryEntitlements(FCommerceQueryEntitlements::Params&& Params)
+{
+	TOnlineAsyncOpRef<FCommerceQueryEntitlements> Op = GetOp<FCommerceQueryEntitlements>(MoveTemp(Params));
+	Op->SetError(Errors::NotImplemented());
+	return Op->GetHandle();
+}
+
+TOnlineResult<FCommerceGetEntitlements> FCommerceEOS::GetEntitlements(FCommerceGetEntitlements::Params&& Params)
+{
+	return TOnlineResult<FCommerceGetEntitlements>(Errors::NotImplemented());
+}
+
+TOnlineAsyncOpHandle<FCommerceRedeemEntitlement> FCommerceEOS::RedeemEntitlement(FCommerceRedeemEntitlement::Params&& Params)
+{
+	TOnlineAsyncOpRef<FCommerceRedeemEntitlement> Op = GetOp<FCommerceRedeemEntitlement>(MoveTemp(Params));
+	Op->SetError(Errors::NotImplemented());
+	return Op->GetHandle();
+}
+
+TOnlineAsyncOpHandle<FCommerceRetrieveS2SToken> FCommerceEOS::RetrieveS2SToken(FCommerceRetrieveS2SToken::Params&& Params)
+{
+	TOnlineAsyncOpRef<FCommerceRetrieveS2SToken> Op = GetOp<FCommerceRetrieveS2SToken>(MoveTemp(Params));
+	Op->SetError(Errors::NotImplemented());
+	return Op->GetHandle();
+}
+
+} // namespace UE::Online
diff --git a/Engine/Plugins/Online/Experimental/OnlineServicesEOS/Source/Private/Online/OnlineServicesEOS.cpp b/Engine/Plugins/Online/Experimental/OnlineServicesEOS/Source/Private/Online/OnlineServicesEOS.cpp
--- a/Engine/Plugins/Online/Experimental/OnlineServicesEOS/Source/Private/Online/OnlineServicesEOS.cpp
+++ b/Engine/Plugins/Online/Experimental/OnlineServicesEOS/Source/Private/Online/OnlineServicesEOS.cpp
@@ -7,6 +7,7 @@
 #include "Online/FriendsEOS.h"
 #include "Online/PresenceEOS.h"
 #include "Online/UserInfoEOS.h"
+#include "Online/CommerceEOS.h"
 
 namespace UE::Online {
 
@@ -22,6 +23,7 @@ void FOnlineServicesEOS::RegisterComponents()
 	Components.Register<FFriendsEOS>(*this);
 	Components.Register<FPresenceEOS>(*this);
 	Components.Register<FUserInfoEOS>(*this);
+	Components.Register<FCommerceEOS>(*this);
 
 	Super::RegisterComponents();
 }
diff --git a/Engine/Plugins/Online/Experimental/OnlineServicesEOS/Source/Public/Online/CommerceEOS.h b/Engine/Plugins/Online/Experimental/OnlineServicesEOS/Source/Public/Online/CommerceEOS.h
new file mode 100644
--- /dev/null
+++ b/Engine/Plugins/Online/Experimental/OnlineServicesEOS/Source/Public/Online/CommerceEOS.h
@@ -0,0 +1,37 @@
+// Copyright Epic Games, Inc. All Rights Reserved.
+
+#pragma once
+
+#include "Online/CommerceCommon.h"
+
+namespace UE::Online {
+
+class FOnlineServicesEOS;
+
+class ONLINESERVICESEOS_API FCommerceEOS : public FCommerceCommon
+{
+public:
+	using Super = FCommerceCommon;
+	using FCommerceCommon::FCommerceCommon;
+
+	// IOnlineComponent
+	virtual void PostInitialize() override;
+	virtual void PreShutdown() override;
+
+	// ICommerce
+	virtual TOnlineAsyncOpHandle<FCommerceQueryOffers> QueryOffers(FCommerceQueryOffers::Params&& Params) override;
+	virtual TOnlineAsyncOpHandle<FCommerceQueryOffersById> QueryOffersById(FCommerceQueryOffersById::Params&& Params) override;
+	virtual TOnlineResult<FCommerceGetOffers> GetOffers(FCommerceGetOffers::Params&& Params) override;
+	virtual TOnlineResult<FCommerceGetOffersById> GetOffersById(FCommerceGetOffersById::Params&& Params) override;
+	virtual TOnlineAsyncOpHandle<FCommerceShowStoreUI> ShowStoreUI(FCommerceShowStoreUI::Params&& Params) override;
+	virtual TOnlineAsyncOpHandle<FCommerceCheckout> Checkout(FCommerceCheckout::Params&& Params) override;
+	virtual TOnlineAsyncOpHandle<FCommerceQueryTransactionEntitlements> QueryTransactionEntitlements(FCommerceQueryTransactionEntitlements::Params&& Params) override;
+	virtual TOnlineAsyncOpHandle<FCommerceQueryEntitlements> QueryEntitlements(FCommerceQueryEntitlements::Params&& Params) override;
+	virtual TOnlineResult<FCommerceGetEntitlements> GetEntitlements(FCommerceGetEntitlements::Params&& Params) override;
+	virtual TOnlineAsyncOpHandle<FCommerceRedeemEntitlement> RedeemEntitlement(FCommerceRedeemEntitlement::Params&& Params) override;
+	virtual TOnlineAsyncOpHandle<FCommerceRetrieveS2SToken> RetrieveS2SToken(FCommerceRetrieveS2SToken::Params&& Params) override;
+
+};
+
+/* UE::Online */
+}
