Commerce EOS full implementation

diff --git a/Engine/Plugins/Online/Experimental/OnlineServices/Source/OnlineServicesCommon/Private/Online/CommerceCommon.cpp b/Engine/Plugins/Online/Experimental/OnlineServices/Source/OnlineServicesCommon/Private/Online/CommerceCommon.cpp
--- a/Engine/Plugins/Online/Experimental/OnlineServices/Source/OnlineServicesCommon/Private/Online/CommerceCommon.cpp
+++ b/Engine/Plugins/Online/Experimental/OnlineServices/Source/OnlineServicesCommon/Private/Online/CommerceCommon.cpp
@@ -103,4 +103,9 @@ TOnlineEvent<void(const FCommerceOnPurchaseComplete&)> FCommerceCommon::OnPurcha
 	return OnPurchaseCompletedEvent;
 }
 
+FText FCommerceCommon::GetFormattedPrice(uint64 Price, int32 DecimalPoint, FString CurrencyCode)
+{
+	return FText::AsCurrencyBase(Price, CurrencyCode, NULL, DecimalPoint);
+}
+
 } // namespace UE::Online
diff --git a/Engine/Plugins/Online/Experimental/OnlineServices/Source/OnlineServicesCommon/Public/Online/CommerceCommon.h b/Engine/Plugins/Online/Experimental/OnlineServices/Source/OnlineServicesCommon/Public/Online/CommerceCommon.h
--- a/Engine/Plugins/Online/Experimental/OnlineServices/Source/OnlineServicesCommon/Public/Online/CommerceCommon.h
+++ b/Engine/Plugins/Online/Experimental/OnlineServices/Source/OnlineServicesCommon/Public/Online/CommerceCommon.h
@@ -33,6 +33,9 @@ public:
 	virtual TOnlineAsyncOpHandle<FCommerceRetrieveS2SToken> RetrieveS2SToken(FCommerceRetrieveS2SToken::Params&& Params) override;
 	virtual TOnlineEvent<void(const FCommerceOnPurchaseComplete&)> OnPurchaseCompleted() override;
 
+	//CommerceCommon
+	FText GetFormattedPrice(uint64 Price, int32 DecimalPoint, FString CurrencyCode);
+
 protected:
 	TOnlineEventCallable<void(const FCommerceOnPurchaseComplete&)> OnPurchaseCompletedEvent;
 };
diff --git a/Engine/Plugins/Online/Experimental/OnlineServicesEOS/Source/Private/Online/CommerceEOS.cpp b/Engine/Plugins/Online/Experimental/OnlineServicesEOS/Source/Private/Online/CommerceEOS.cpp
--- a/Engine/Plugins/Online/Experimental/OnlineServicesEOS/Source/Private/Online/CommerceEOS.cpp
+++ b/Engine/Plugins/Online/Experimental/OnlineServicesEOS/Source/Private/Online/CommerceEOS.cpp
@@ -1,15 +1,22 @@
 // Copyright Epic Games, Inc. All Rights Reserved.
 
 #include "Online/CommerceEOS.h"
+#include "Online/OnlineServicesEOS.h"
 #include "Online/OnlineAsyncOp.h"
 #include "Online/OnlineErrorDefinitions.h"
 #include "Online/OnlineServicesCommon.h"
+#include "Online/OnlineIdEOS.h"
+#include "Online/AuthEOS.h"
+#include "eos_ecom.h"
 
 namespace UE::Online {
 
 void FCommerceEOS::PostInitialize()
 {	
 	Super::PostInitialize();
+
+	EcomHandle = EOS_Platform_GetEcomInterface(static_cast<FOnlineServicesEOS&>(GetServices()).GetEOSPlatformHandle());
+	check(EcomHandle);
 }
 
 void FCommerceEOS::PreShutdown()
@@ -17,32 +24,149 @@ void FCommerceEOS::PreShutdown()
 	Super::PreShutdown();
 }
 
+void FCommerceEOS::EOSOfferToOssOffer(FOffer& OutOffer, EOS_Ecom_CatalogOffer* EosOffer)
+{
+	OutOffer.OfferId = UTF8_TO_TCHAR(EosOffer->Id);
+	OutOffer.Title = FText::FromString(UTF8_TO_TCHAR(EosOffer->TitleText));
+	OutOffer.LongDescription = FText::FromString(UTF8_TO_TCHAR(EosOffer->LongDescriptionText));
+	OutOffer.ExpirationDate = FDateTime(EosOffer->ExpirationTimestamp);
+	OutOffer.CurrencyCode = UTF8_TO_TCHAR(EosOffer->CurrencyCode);
+	OutOffer.PurchaseLimit = EosOffer->PurchaseLimit;
+	if (EosOffer->PriceResult == EOS_EResult::EOS_Success)
+	{
+		OutOffer.RegularPrice = EosOffer->OriginalPrice64;
+		OutOffer.Price = EosOffer->CurrentPrice64;
+		OutOffer.PriceDecimalPoint = EosOffer->DecimalPoint;
+
+		OutOffer.FormattedPrice = GetFormattedPrice(OutOffer.Price, OutOffer.PriceDecimalPoint, OutOffer.CurrencyCode);
+		OutOffer.FormattedRegularPrice = GetFormattedPrice(OutOffer.RegularPrice, OutOffer.PriceDecimalPoint, OutOffer.CurrencyCode);
+	}
+	// Other properties that could potentially be added in the future- image data (EOS_Ecom_KeyImageInfo ), individual items (EOS_Ecom_CatalogItem)
+}
+
+void FCommerceEOS::EOSEntitlementToOssEntitlement(FEntitlement& OutEntitlement, EOS_Ecom_Entitlement* EosEntitlement)
+{
+	OutEntitlement.EntitlementId = UTF8_TO_TCHAR(EosEntitlement->EntitlementId);
+	OutEntitlement.bRedeemed = (EosEntitlement->bRedeemed == EOS_TRUE);
+	OutEntitlement.ProductId = UTF8_TO_TCHAR(EosEntitlement->CatalogItemId);
+	OutEntitlement.Quantity = 1;
+	// do we have a suitable placement for Name or AcquiredDate here?
+}
+
 TOnlineAsyncOpHandle<FCommerceQueryOffers> FCommerceEOS::QueryOffers(FCommerceQueryOffers::Params&& Params)
 {
 	TOnlineAsyncOpRef<FCommerceQueryOffers> Op = GetOp<FCommerceQueryOffers>(MoveTemp(Params));
-	Op->SetError(Errors::NotImplemented());
+
+	Op->Then([this](TOnlineAsyncOp<FCommerceQueryOffers>& Op, TPromise<const EOS_Ecom_QueryOffersCallbackInfo*>&& Promise)
+	{
+		const FCommerceQueryOffers::Params& Params = Op.GetParams();
+		if (!Services.Get<FAuthEOS>()->IsLoggedIn(Params.LocalUserId))
+		{
+			Op.SetError(Errors::NotLoggedIn());
+			return;
+		}
+		EOS_EpicAccountId LocalUserEasId = GetEpicAccountId(Params.LocalUserId);
+		if (!EOS_EpicAccountId_IsValid(LocalUserEasId))
+		{
+			Op.SetError(Errors::NotLoggedIn());
+			return;
+		}
+		
+		EOS_Ecom_QueryOffersOptions Options = { };
+		Options.ApiVersion = EOS_ECOM_QUERYOFFERS_API_LATEST;
+		Options.LocalUserId = LocalUserEasId;
+		EOS_Async(EOS_Ecom_QueryOffers, EcomHandle, Options, MoveTemp(Promise));
+	})
+	.Then([this](TOnlineAsyncOp<FCommerceQueryOffers>& Op, const EOS_Ecom_QueryOffersCallbackInfo* Data)
+	{
+		EOS_EResult Result = Data->ResultCode;
+		if (Result != EOS_EResult::EOS_Success)
+		{
+			Op.SetError(Errors::FromEOSResult(Result));
+			return;
+		}
+
+		EOS_Ecom_GetOfferCountOptions CountOptions = { };
+		CountOptions.ApiVersion = EOS_ECOM_GETOFFERCOUNT_API_LATEST;
+		CountOptions.LocalUserId = Data->LocalUserId;
+		uint32 OfferCount = EOS_Ecom_GetOfferCount(EcomHandle, &CountOptions);
+
+		TArray<FOffer>& Offers = CachedOffers.FindOrAdd(Op.GetParams().LocalUserId, TArray<FOffer>());
+		Offers.Empty(OfferCount);
+
+		EOS_Ecom_CopyOfferByIndexOptions OfferOptions = { };
+		OfferOptions.ApiVersion = EOS_ECOM_COPYOFFERBYINDEX_API_LATEST;
+		OfferOptions.LocalUserId = Data->LocalUserId;
+		// Iterate and parse the offer list
+		for (uint32 OfferIndex = 0; OfferIndex < OfferCount; OfferIndex++)
+		{
+			EOS_Ecom_CatalogOffer* EosOffer = nullptr;
+			OfferOptions.OfferIndex = OfferIndex;
+			EOS_EResult OfferResult = EOS_Ecom_CopyOfferByIndex(EcomHandle, &OfferOptions, &EosOffer);
+			if (OfferResult != EOS_EResult::EOS_Success)
+			{
+				continue;
+			}
+			FOffer& Offer = Offers.Emplace_GetRef();
+			EOSOfferToOssOffer(Offer, EosOffer);
+			EOS_Ecom_CatalogOffer_Release(EosOffer);
+		}
+
+		Op.SetResult({});
+	})
+	.Enqueue(GetSerialQueue());
+
 	return Op->GetHandle();
 }
 
 TOnlineAsyncOpHandle<FCommerceQueryOffersById> FCommerceEOS::QueryOffersById(FCommerceQueryOffersById::Params&& Params)
 {
 	TOnlineAsyncOpRef<FCommerceQueryOffersById> Op = GetOp<FCommerceQueryOffersById>(MoveTemp(Params));
-	Op->SetError(Errors::NotImplemented());
+
+	// EOS doesn't support anything ID-specific, just use the generic QueryOffers call
+	QueryOffers({Op->GetParams().LocalUserId})
+	.OnComplete([this, Op](const TOnlineResult<FCommerceQueryOffers>& Result)
+	{
+		if(Result.IsOk())
+		{
+			Op->SetResult({});
+		}
+		else
+		{
+			FOnlineError Error = Result.GetErrorValue();
+			Op->SetError(MoveTemp(Error));
+		}
+	});	
+
 	return Op->GetHandle();
 }
 
 TOnlineResult<FCommerceGetOffers> FCommerceEOS::GetOffers(FCommerceGetOffers::Params&& Params)
 {
-	return TOnlineResult<FCommerceGetOffers>(Errors::NotImplemented());
+	if(CachedOffers.Contains(Params.LocalUserId))
+	{
+		return TOnlineResult<FCommerceGetOffers>({CachedOffers.FindChecked(Params.LocalUserId)});
+	}
+	return TOnlineResult<FCommerceGetOffers>(Errors::NotFound());
 }
 
 TOnlineResult<FCommerceGetOffersById> FCommerceEOS::GetOffersById(FCommerceGetOffersById::Params&& Params)
 {
-	return TOnlineResult<FCommerceGetOffersById>(Errors::NotImplemented());
+	if (CachedOffers.Contains(Params.LocalUserId))
+	{
+		return TOnlineResult<FCommerceGetOffersById>({ CachedOffers.FindChecked(Params.LocalUserId).FilterByPredicate(
+			[&Params](const FOffer& Offer)
+			{ 
+				return Params.OfferIds.Contains(Offer.OfferId);
+			})
+		});
+	}
+	return TOnlineResult<FCommerceGetOffersById>(Errors::NotFound());
 }
 
 TOnlineAsyncOpHandle<FCommerceShowStoreUI> FCommerceEOS::ShowStoreUI(FCommerceShowStoreUI::Params&& Params)
 {
+	// Todo: Implement when EOS store overlay is fully supported
 	TOnlineAsyncOpRef<FCommerceShowStoreUI> Op = GetOp<FCommerceShowStoreUI>(MoveTemp(Params));
 	Op->SetError(Errors::NotImplemented());
 	return Op->GetHandle();
@@ -51,40 +175,306 @@ TOnlineAsyncOpHandle<FCommerceShowStoreUI> FCommerceEOS::ShowStoreUI(FCommerceSh
 TOnlineAsyncOpHandle<FCommerceCheckout> FCommerceEOS::Checkout(FCommerceCheckout::Params&& Params)
 {
 	TOnlineAsyncOpRef<FCommerceCheckout> Op = GetOp<FCommerceCheckout>(MoveTemp(Params));
-	Op->SetError(Errors::NotImplemented());
+
+	Op->Then([this](TOnlineAsyncOp<FCommerceCheckout>& Op, TPromise<const EOS_Ecom_CheckoutCallbackInfo*>&& Promise)
+	{
+		const FCommerceCheckout::Params& Params = Op.GetParams();
+		if (!Services.Get<FAuthEOS>()->IsLoggedIn(Params.LocalUserId))
+		{
+			Op.SetError(Errors::NotLoggedIn());
+			return;
+		}
+		EOS_EpicAccountId LocalUserEasId = GetEpicAccountId(Params.LocalUserId);
+		if (!EOS_EpicAccountId_IsValid(LocalUserEasId))
+		{
+			Op.SetError(Errors::NotLoggedIn());
+			return;
+		}
+
+		EOS_Ecom_CheckoutEntry* EosCheckoutEntries = new EOS_Ecom_CheckoutEntry[Params.Offers.Num()];
+		for(size_t i = 0; i < Params.Offers.Num(); i++)
+		{
+			FTCHARToUTF8 Utf8CheckoutId(*Params.Offers[i].OfferId);
+
+			EosCheckoutEntries[i] = EOS_Ecom_CheckoutEntry();
+			EosCheckoutEntries[i].ApiVersion = EOS_ECOM_CHECKOUTENTRY_API_LATEST;
+			EosCheckoutEntries[i].OfferId = Utf8CheckoutId.Get();
+		}
+
+		EOS_Ecom_CheckoutOptions Options = { };
+		Options.ApiVersion = EOS_ECOM_CHECKOUT_API_LATEST;
+		Options.LocalUserId = LocalUserEasId;
+		Options.EntryCount = Params.Offers.Num();
+		Options.Entries = EosCheckoutEntries;
+
+		EOS_Async(EOS_Ecom_Checkout, EcomHandle, Options, MoveTemp(Promise));
+		delete[] EosCheckoutEntries;
+	})
+	.Then([this](TOnlineAsyncOp<FCommerceCheckout>& Op, const EOS_Ecom_CheckoutCallbackInfo* Data)
+	{
+		EOS_EResult Result = Data->ResultCode;
+		if (Result != EOS_EResult::EOS_Success)
+		{
+			UE_LOG(LogTemp, Error, TEXT("EOS_Ecom_Checkout: failed with error (%s)"), ANSI_TO_TCHAR(EOS_EResult_ToString(Data->ResultCode)));
+			Op.SetError(Errors::FromEOSResult(Result));
+			return;
+		}
+
+		FCommerceCheckout::Result ResultValue;
+		ResultValue.TransactionId.Emplace(ANSI_TO_TCHAR(Data->TransactionId));
+		Op.SetResult(MoveTemp(ResultValue));
+	})
+	.Enqueue(GetSerialQueue());
+
 	return Op->GetHandle();
 }
 
 TOnlineAsyncOpHandle<FCommerceQueryTransactionEntitlements> FCommerceEOS::QueryTransactionEntitlements(FCommerceQueryTransactionEntitlements::Params&& Params)
 {
 	TOnlineAsyncOpRef<FCommerceQueryTransactionEntitlements> Op = GetOp<FCommerceQueryTransactionEntitlements>(MoveTemp(Params));
-	Op->SetError(Errors::NotImplemented());
+
+	Op->Then([this](TOnlineAsyncOp<FCommerceQueryTransactionEntitlements>& Op)
+	{
+		const FCommerceQueryTransactionEntitlements::Params& Params = Op.GetParams();
+		if (!Services.Get<FAuthEOS>()->IsLoggedIn(Params.LocalUserId))
+		{
+			Op.SetError(Errors::NotLoggedIn());
+			return;
+		}
+		EOS_EpicAccountId LocalUserEasId = GetEpicAccountId(Params.LocalUserId);
+		if (!EOS_EpicAccountId_IsValid(LocalUserEasId))
+		{
+			Op.SetError(Errors::NotLoggedIn());
+			return;
+		}
+
+
+		EOS_Ecom_HTransaction OutTransaction;
+		FTCHARToUTF8 Utf8TransactionId(*Params.TransactionId);
+		EOS_Ecom_CopyTransactionByIdOptions CopyTransactionOptions = { };
+		CopyTransactionOptions.ApiVersion = EOS_ECOM_COPYTRANSACTIONBYID_API_LATEST;
+		CopyTransactionOptions.LocalUserId = LocalUserEasId;
+		CopyTransactionOptions.TransactionId = Utf8TransactionId.Get();
+
+		EOS_EResult CopyTransactionResult = EOS_Ecom_CopyTransactionById(EcomHandle, &CopyTransactionOptions, &OutTransaction);
+		if (CopyTransactionResult != EOS_EResult::EOS_Success)
+		{
+			UE_LOG(LogTemp, Error, TEXT("EOS_Ecom_CopyTransactionById: failed with error (%s)"), ANSI_TO_TCHAR(EOS_EResult_ToString(CopyTransactionResult)));
+			Op.SetError(Errors::FromEOSResult(CopyTransactionResult));
+			return;
+		}
+
+		EOS_Ecom_Transaction_GetEntitlementsCountOptions EntitlementCountOptions = { };
+		EntitlementCountOptions.ApiVersion = EOS_ECOM_TRANSACTION_GETENTITLEMENTSCOUNT_API_LATEST;
+		int32 NumTransactionsFound = EOS_Ecom_Transaction_GetEntitlementsCount(OutTransaction, &EntitlementCountOptions);
+
+		if(NumTransactionsFound > 0)
+		{
+			FCommerceQueryTransactionEntitlements::Result OutResult;
+
+			for(int32 i = 0; i < NumTransactionsFound; i++)
+			{
+				EOS_Ecom_Entitlement* EosEntitlement;
+				EOS_Ecom_Transaction_CopyEntitlementByIndexOptions CopyTransactionEntitlementOptions = {};
+				CopyTransactionEntitlementOptions.ApiVersion = EOS_ECOM_TRANSACTION_COPYENTITLEMENTBYINDEX_API_LATEST;
+				CopyTransactionEntitlementOptions.EntitlementIndex = i;
+
+				EOS_EResult Result = EOS_Ecom_Transaction_CopyEntitlementByIndex(OutTransaction, &CopyTransactionEntitlementOptions, &EosEntitlement);
+				if(Result == EOS_EResult::EOS_Success)
+				{
+					FEntitlement& Entitlement = OutResult.Entitlements.Emplace_GetRef();
+					EOSEntitlementToOssEntitlement(Entitlement, EosEntitlement);
+					EOS_Ecom_Entitlement_Release(EosEntitlement);
+				}
+			}
+
+			EOS_Ecom_Transaction_Release(OutTransaction);
+			Op.SetResult(MoveTemp(OutResult));
+		}
+		else
+		{
+			EOS_Ecom_Transaction_Release(OutTransaction);
+			Op.SetError(Errors::NotFound());
+		}
+	})
+	.Enqueue(GetSerialQueue());
+
 	return Op->GetHandle();
 }
 
 TOnlineAsyncOpHandle<FCommerceQueryEntitlements> FCommerceEOS::QueryEntitlements(FCommerceQueryEntitlements::Params&& Params)
 {
 	TOnlineAsyncOpRef<FCommerceQueryEntitlements> Op = GetOp<FCommerceQueryEntitlements>(MoveTemp(Params));
-	Op->SetError(Errors::NotImplemented());
+
+	Op->Then([this](TOnlineAsyncOp<FCommerceQueryEntitlements>& Op, TPromise<const EOS_Ecom_QueryEntitlementsCallbackInfo*>&& Promise)
+	{
+		const FCommerceQueryEntitlements::Params& Params = Op.GetParams();
+		if (!Services.Get<FAuthEOS>()->IsLoggedIn(Params.LocalUserId))
+		{
+			Op.SetError(Errors::NotLoggedIn());
+			return;
+		}
+		EOS_EpicAccountId LocalUserEasId = GetEpicAccountId(Params.LocalUserId);
+		if (!EOS_EpicAccountId_IsValid(LocalUserEasId))
+		{
+			Op.SetError(Errors::NotLoggedIn());
+			return;
+		}
+
+		EOS_Ecom_QueryEntitlementsOptions Options = { };
+		Options.ApiVersion = EOS_ECOM_QUERYENTITLEMENTS_API_LATEST;
+		Options.LocalUserId = LocalUserEasId;
+		Options.bIncludeRedeemed = EOS_TRUE;
+
+		EOS_Async(EOS_Ecom_QueryEntitlements, EcomHandle, Options, MoveTemp(Promise));
+	})
+	.Then([this](TOnlineAsyncOp<FCommerceQueryEntitlements>& Op, const EOS_Ecom_QueryEntitlementsCallbackInfo* Data)
+	{
+		EOS_EResult Result = Data->ResultCode;
+		if (Result != EOS_EResult::EOS_Success)
+		{
+			UE_LOG(LogTemp, Error, TEXT("EOS_Ecom_QueryEntitlements: failed with error (%s)"), ANSI_TO_TCHAR(EOS_EResult_ToString(Data->ResultCode)));
+			Op.SetError(Errors::FromEOSResult(Result));
+			return;
+		}
+
+		EOS_Ecom_GetEntitlementsCountOptions CountOptions = { };
+		CountOptions.ApiVersion = EOS_ECOM_GETENTITLEMENTSCOUNT_API_LATEST;
+		CountOptions.LocalUserId = Data->LocalUserId;
+		uint32 Count = EOS_Ecom_GetEntitlementsCount(EcomHandle, &CountOptions);
+		TArray<FEntitlement>& Entitlements = CachedEntitlements.FindOrAdd(Op.GetParams().LocalUserId, TArray<FEntitlement>());
+		Entitlements.Reset(Count);
+
+		EOS_Ecom_CopyEntitlementByIndexOptions CopyOptions = { };
+		CopyOptions.ApiVersion = EOS_ECOM_COPYENTITLEMENTBYINDEX_API_LATEST;
+		CopyOptions.LocalUserId = Data->LocalUserId;
+
+		for (uint32 Index = 0; Index < Count; Index++)
+		{
+			CopyOptions.EntitlementIndex = Index;
+
+			EOS_Ecom_Entitlement* EosEntitlement = nullptr;
+			EOS_EResult CopyResult = EOS_Ecom_CopyEntitlementByIndex(EcomHandle, &CopyOptions, &EosEntitlement);
+			if (CopyResult != EOS_EResult::EOS_Success && CopyResult != EOS_EResult::EOS_Ecom_EntitlementStale)
+			{
+				UE_LOG(LogTemp, Error, TEXT("EOS_Ecom_CopyEntitlementByIndex: failed with error (%s) (proceeding with operation)"), ANSI_TO_TCHAR(EOS_EResult_ToString(CopyResult)));
+				continue;
+			}
+
+			// Parse the entitlement into the receipt format
+			FEntitlement Entitlement = Entitlements.Emplace_GetRef();
+			EOSEntitlementToOssEntitlement(Entitlement, EosEntitlement);
+
+			EOS_Ecom_Entitlement_Release(EosEntitlement);
+		}
+
+		Op.SetResult({});
+	})
+	.Enqueue(GetSerialQueue());
+
 	return Op->GetHandle();
 }
 
 TOnlineResult<FCommerceGetEntitlements> FCommerceEOS::GetEntitlements(FCommerceGetEntitlements::Params&& Params)
 {
-	return TOnlineResult<FCommerceGetEntitlements>(Errors::NotImplemented());
+	if (!CachedEntitlements.Contains(Params.LocalUserId))
+	{
+		return TOnlineResult<FCommerceGetEntitlements>(Errors::NotFound());
+	}
+	return TOnlineResult<FCommerceGetEntitlements>({CachedEntitlements.FindChecked(Params.LocalUserId)});
 }
 
 TOnlineAsyncOpHandle<FCommerceRedeemEntitlement> FCommerceEOS::RedeemEntitlement(FCommerceRedeemEntitlement::Params&& Params)
 {
 	TOnlineAsyncOpRef<FCommerceRedeemEntitlement> Op = GetOp<FCommerceRedeemEntitlement>(MoveTemp(Params));
-	Op->SetError(Errors::NotImplemented());
+
+	Op->Then([this](TOnlineAsyncOp<FCommerceRedeemEntitlement>& Op, TPromise<const EOS_Ecom_RedeemEntitlementsCallbackInfo*>&& Promise)
+	{
+		const FCommerceRedeemEntitlement::Params& Params = Op.GetParams();
+		if (!Services.Get<FAuthEOS>()->IsLoggedIn(Params.LocalUserId))
+		{
+			Op.SetError(Errors::NotLoggedIn());
+			return;
+		}
+		EOS_EpicAccountId LocalUserEasId = GetEpicAccountId(Params.LocalUserId);
+		if (!EOS_EpicAccountId_IsValid(LocalUserEasId))
+		{
+			Op.SetError(Errors::NotLoggedIn());
+			return;
+		}
+
+		FTCHARToUTF8 Utf8EntitlementId(*Op.GetParams().EntitlementId);
+		const char* EntitlementIdPtr = Utf8EntitlementId.Get();
+		EOS_Ecom_RedeemEntitlementsOptions Options = { };
+		Options.ApiVersion = EOS_ECOM_REDEEMENTITLEMENTS_API_LATEST;
+		Options.LocalUserId = LocalUserEasId;
+		Options.EntitlementIdCount = 1;
+		Options.EntitlementIds = &EntitlementIdPtr;
+		
+		EOS_Async(EOS_Ecom_RedeemEntitlements, EcomHandle, Options, MoveTemp(Promise));
+	})
+	.Then([this](TOnlineAsyncOp<FCommerceRedeemEntitlement>& Op, const EOS_Ecom_RedeemEntitlementsCallbackInfo* Data)
+	{
+		EOS_EResult Result = Data->ResultCode;
+		if (Result != EOS_EResult::EOS_Success)
+		{
+			UE_LOG(LogTemp, Error, TEXT("EOS_Ecom_RedeemEntitlements: failed with error (%s)"), ANSI_TO_TCHAR(EOS_EResult_ToString(Data->ResultCode)));
+			Op.SetError(Errors::FromEOSResult(Result));
+			return;
+		}
+		Op.SetResult({});
+	})
+	.Enqueue(GetSerialQueue());
+
 	return Op->GetHandle();
 }
 
 TOnlineAsyncOpHandle<FCommerceRetrieveS2SToken> FCommerceEOS::RetrieveS2SToken(FCommerceRetrieveS2SToken::Params&& Params)
 {
 	TOnlineAsyncOpRef<FCommerceRetrieveS2SToken> Op = GetOp<FCommerceRetrieveS2SToken>(MoveTemp(Params));
-	Op->SetError(Errors::NotImplemented());
+
+	Op->Then([this](TOnlineAsyncOp<FCommerceRetrieveS2SToken>& Op, TPromise<const EOS_Ecom_QueryOwnershipTokenCallbackInfo*>&& Promise)
+	{
+		const FCommerceRetrieveS2SToken::Params& Params = Op.GetParams();
+		if (!Services.Get<FAuthEOS>()->IsLoggedIn(Params.LocalUserId))
+		{
+			Op.SetError(Errors::NotLoggedIn());
+			return;
+		}
+		EOS_EpicAccountId LocalUserEasId = GetEpicAccountId(Params.LocalUserId);
+		if (!EOS_EpicAccountId_IsValid(LocalUserEasId))
+		{
+			Op.SetError(Errors::NotLoggedIn());
+			return;
+		}
+
+		EOS_Ecom_QueryOwnershipTokenOptions Options = {};
+		Options.ApiVersion = EOS_ECOM_QUERYOWNERSHIPTOKEN_API_LATEST;
+		Options.LocalUserId = LocalUserEasId;
+
+		EOS_Async(EOS_Ecom_QueryOwnershipToken, EcomHandle, Options, MoveTemp(Promise));
+	})
+	.Then([this](TOnlineAsyncOp<FCommerceRetrieveS2SToken>& Op, const EOS_Ecom_QueryOwnershipTokenCallbackInfo* Data)
+	{
+		EOS_EResult Result = Data->ResultCode;
+		if (Result != EOS_EResult::EOS_Success)
+		{
+			UE_LOG(LogTemp, Error, TEXT("EOS_Ecom_QueryOwnershipToken: failed with error (%s)"), ANSI_TO_TCHAR(EOS_EResult_ToString(Data->ResultCode)));
+			Op.SetError(Errors::FromEOSResult(Result));
+			return;
+		}
+
+		size_t DataSize = strlen(Data->OwnershipToken);
+		FCommerceRetrieveS2SToken::Result ResultValue;
+		ResultValue.Token.AddZeroed(DataSize);
+		for(size_t i = 0; i < DataSize; i++)
+		{
+			ResultValue.Token[i] = (uint8) Data->OwnershipToken[i];
+		}
+		Op.SetResult(MoveTemp(ResultValue));
+	})
+	.Enqueue(GetSerialQueue());
+
 	return Op->GetHandle();
 }
 
diff --git a/Engine/Plugins/Online/Experimental/OnlineServicesEOS/Source/Public/Online/CommerceEOS.h b/Engine/Plugins/Online/Experimental/OnlineServicesEOS/Source/Public/Online/CommerceEOS.h
--- a/Engine/Plugins/Online/Experimental/OnlineServicesEOS/Source/Public/Online/CommerceEOS.h
+++ b/Engine/Plugins/Online/Experimental/OnlineServicesEOS/Source/Public/Online/CommerceEOS.h
@@ -3,6 +3,7 @@
 #pragma once
 
 #include "Online/CommerceCommon.h"
+#include "eos_ecom.h"
 
 namespace UE::Online {
 
@@ -31,6 +32,13 @@ public:
 	virtual TOnlineAsyncOpHandle<FCommerceRedeemEntitlement> RedeemEntitlement(FCommerceRedeemEntitlement::Params&& Params) override;
 	virtual TOnlineAsyncOpHandle<FCommerceRetrieveS2SToken> RetrieveS2SToken(FCommerceRetrieveS2SToken::Params&& Params) override;
 
+private:
+	void EOSOfferToOssOffer(FOffer& InOffer, EOS_Ecom_CatalogOffer* EosOffer);
+	void EOSEntitlementToOssEntitlement(FEntitlement& OutEntitlement, EOS_Ecom_Entitlement* EosEntitlement);
+
+	EOS_HEcom EcomHandle;
+	TMap<FOnlineAccountIdHandle, TArray<FEntitlement>> CachedEntitlements;
+	TMap<FOnlineAccountIdHandle, TArray<FOffer>> CachedOffers;
 };
 
 /* UE::Online */
