Fix non-unity compile error: include platform file with CommerceEOS.h

Horde issue 203022

diff --git a/Engine/Plugins/Online/Experimental/OnlineServicesEOS/Source/Public/Online/CommerceEOS.h b/Engine/Plugins/Online/Experimental/OnlineServicesEOS/Source/Public/Online/CommerceEOS.h
--- a/Engine/Plugins/Online/Experimental/OnlineServicesEOS/Source/Public/Online/CommerceEOS.h
+++ b/Engine/Plugins/Online/Experimental/OnlineServicesEOS/Source/Public/Online/CommerceEOS.h
@@ -3,6 +3,10 @@
 #pragma once
 
 #include "Online/CommerceCommon.h"
+
+#if defined(EOS_PLATFORM_BASE_FILE_NAME)
+#include EOS_PLATFORM_BASE_FILE_NAME
+#endif
 #include "eos_ecom.h"
 
 namespace UE::Online {
