Speculative fix for static analysis warning

diff --git a/Engine/Plugins/Online/Experimental/OnlineServicesEOS/Source/Private/Online/CommerceEOS.cpp b/Engine/Plugins/Online/Experimental/OnlineServicesEOS/Source/Private/Online/CommerceEOS.cpp
--- a/Engine/Plugins/Online/Experimental/OnlineServicesEOS/Source/Private/Online/CommerceEOS.cpp
+++ b/Engine/Plugins/Online/Experimental/OnlineServicesEOS/Source/Private/Online/CommerceEOS.cpp
@@ -192,7 +192,7 @@ TOnlineAsyncOpHandle<FCommerceCheckout> FCommerceEOS::Checkout(FCommerceCheckout
 		}
 
 		EOS_Ecom_CheckoutEntry* EosCheckoutEntries = new EOS_Ecom_CheckoutEntry[Params.Offers.Num()];
-		for(size_t i = 0; i < Params.Offers.Num(); i++)
+		for (int32 i = 0; i < Params.Offers.Num(); i++)
 		{
 			FTCHARToUTF8 Utf8CheckoutId(*Params.Offers[i].OfferId);
 
