Fix holding a pointer to an out of scope/destructed object

diff --git a/Engine/Plugins/Online/Experimental/OnlineServicesEOS/Source/Private/Online/CommerceEOS.cpp b/Engine/Plugins/Online/Experimental/OnlineServicesEOS/Source/Private/Online/CommerceEOS.cpp
--- a/Engine/Plugins/Online/Experimental/OnlineServicesEOS/Source/Private/Online/CommerceEOS.cpp
+++ b/Engine/Plugins/Online/Experimental/OnlineServicesEOS/Source/Private/Online/CommerceEOS.cpp
@@ -191,14 +191,16 @@ TOnlineAsyncOpHandle<FCommerceCheckout> FCommerceEOS::Checkout(FCommerceCheckout
 			return;
 		}
 
+		TArray<FTCHARToUTF8> Utf8CheckoutIds;
+		Utf8CheckoutIds.Reserve(Params.Offers.Num());
+
 		EOS_Ecom_CheckoutEntry* EosCheckoutEntries = new EOS_Ecom_CheckoutEntry[Params.Offers.Num()];
 		for (int32 i = 0; i < Params.Offers.Num(); i++)
 		{
-			FTCHARToUTF8 Utf8CheckoutId(*Params.Offers[i].OfferId);
+			Utf8CheckoutIds.Emplace(*Params.Offers[i].OfferId);
 
-			EosCheckoutEntries[i] = EOS_Ecom_CheckoutEntry();
 			EosCheckoutEntries[i].ApiVersion = EOS_ECOM_CHECKOUTENTRY_API_LATEST;
-			EosCheckoutEntries[i].OfferId = Utf8CheckoutId.Get();
+			EosCheckoutEntries[i].OfferId = Utf8CheckoutIds.Last().Get();
 		}
 
 		EOS_Ecom_CheckoutOptions Options = { };
