Restructure code to speculatively silence static analysis warning

diff --git a/Engine/Plugins/Online/Experimental/OnlineServicesEOS/Source/Private/Online/CommerceEOS.cpp b/Engine/Plugins/Online/Experimental/OnlineServicesEOS/Source/Private/Online/CommerceEOS.cpp
--- a/Engine/Plugins/Online/Experimental/OnlineServicesEOS/Source/Private/Online/CommerceEOS.cpp
+++ b/Engine/Plugins/Online/Experimental/OnlineServicesEOS/Source/Private/Online/CommerceEOS.cpp
@@ -194,23 +194,25 @@ TOnlineAsyncOpHandle<FCommerceCheckout> FCommerceEOS::Checkout(FCommerceCheckout
 		TArray<FTCHARToUTF8> Utf8CheckoutIds;
 		Utf8CheckoutIds.Reserve(Params.Offers.Num());
 
-		EOS_Ecom_CheckoutEntry* EosCheckoutEntries = new EOS_Ecom_CheckoutEntry[Params.Offers.Num()];
-		for (int32 i = 0; i < Params.Offers.Num(); i++)
+		TArray<EOS_Ecom_CheckoutEntry> EosCheckoutEntries;
+		EosCheckoutEntries.Reserve(Params.Offers.Num());
+
+		for (const FPurchaseOffer& Offer : Params.Offers)
 		{
-			Utf8CheckoutIds.Emplace(*Params.Offers[i].OfferId);
+			Utf8CheckoutIds.Emplace(*Offer.OfferId);
 
-			EosCheckoutEntries[i].ApiVersion = EOS_ECOM_CHECKOUTENTRY_API_LATEST;
-			EosCheckoutEntries[i].OfferId = Utf8CheckoutIds.Last().Get();
+			EosCheckoutEntries.AddDefaulted();
+			EosCheckoutEntries.Last().ApiVersion = EOS_ECOM_CHECKOUTENTRY_API_LATEST;
+			EosCheckoutEntries.Last().OfferId = Utf8CheckoutIds.Last().Get();
 		}
 
 		EOS_Ecom_CheckoutOptions Options = { };
 		Options.ApiVersion = EOS_ECOM_CHECKOUT_API_LATEST;
 		Options.LocalUserId = LocalUserEasId;
 		Options.EntryCount = Params.Offers.Num();
-		Options.Entries = EosCheckoutEntries;
+		Options.Entries = EosCheckoutEntries.GetData();
 
 		EOS_Async(EOS_Ecom_Checkout, EcomHandle, Options, MoveTemp(Promise));
-		delete[] EosCheckoutEntries;
 	})
 	.Then([this](TOnlineAsyncOp<FCommerceCheckout>& Op, const EOS_Ecom_CheckoutCallbackInfo* Data)
 	{
