- Use joinable ops in EOS Auth

diff --git a/Engine/Plugins/Online/Experimental/OnlineServicesEOS/Source/Private/Online/AuthEOS.cpp b/Engine/Plugins/Online/Experimental/OnlineServicesEOS/Source/Private/Online/AuthEOS.cpp
--- a/Engine/Plugins/Online/Experimental/OnlineServicesEOS/Source/Private/Online/AuthEOS.cpp
+++ b/Engine/Plugins/Online/Experimental/OnlineServicesEOS/Source/Private/Online/AuthEOS.cpp
@@ -277,72 +277,75 @@ TFuture<FAccountId> ResolveAccountIdImpl(FAuthEOS& AuthEOS, const FAccountId& Lo
 
 TOnlineAsyncOpHandle<FAuthQueryExternalServerAuthTicket> FAuthEOS::QueryExternalServerAuthTicket(FAuthQueryExternalServerAuthTicket::Params&& Params)
 {
-	TOnlineAsyncOpRef<FAuthQueryExternalServerAuthTicket> Op = GetOp<FAuthQueryExternalServerAuthTicket>(MoveTemp(Params));
-
-	Op->Then([this](TOnlineAsyncOp<FAuthQueryExternalServerAuthTicket>& InAsyncOp)
+	TOnlineAsyncOpRef<FAuthQueryExternalServerAuthTicket> Op = GetJoinableOp<FAuthQueryExternalServerAuthTicket>(MoveTemp(Params));
+	if (!Op->IsReady())
 	{
-		const FAuthQueryExternalServerAuthTicket::Params& Params = InAsyncOp.GetParams();
-		TSharedPtr<FAccountInfoEOS> AccountInfoEOS = AccountInfoRegistryEOS.Find(Params.LocalAccountId);
-		if (!AccountInfoEOS)
+		Op->Then([this](TOnlineAsyncOp<FAuthQueryExternalServerAuthTicket>& InAsyncOp)
 		{
-			InAsyncOp.SetError(Errors::InvalidParams());
-			return;
-		}
+			const FAuthQueryExternalServerAuthTicket::Params& Params = InAsyncOp.GetParams();
+			TSharedPtr<FAccountInfoEOS> AccountInfoEOS = AccountInfoRegistryEOS.Find(Params.LocalAccountId);
+			if (!AccountInfoEOS)
+			{
+				InAsyncOp.SetError(Errors::InvalidParams());
+				return;
+			}
 
-		EOS_Auth_CopyUserAuthTokenOptions CopyUserAuthTokenOptions = {};
-		CopyUserAuthTokenOptions.ApiVersion = EOS_AUTH_COPYUSERAUTHTOKEN_API_LATEST;
-		static_assert(EOS_AUTH_COPYUSERAUTHTOKEN_API_LATEST == 1, "EOS_Auth_CopyUserAuthTokenOptions updated, check new fields");
+			EOS_Auth_CopyUserAuthTokenOptions CopyUserAuthTokenOptions = {};
+			CopyUserAuthTokenOptions.ApiVersion = EOS_AUTH_COPYUSERAUTHTOKEN_API_LATEST;
+			static_assert(EOS_AUTH_COPYUSERAUTHTOKEN_API_LATEST == 1, "EOS_Auth_CopyUserAuthTokenOptions updated, check new fields");
 
-		EOS_Auth_Token* AuthToken = nullptr;
+			EOS_Auth_Token* AuthToken = nullptr;
 
-		EOS_EResult Result = EOS_Auth_CopyUserAuthToken(AuthHandle, &CopyUserAuthTokenOptions, AccountInfoEOS->EpicAccountId, &AuthToken);
-		if (Result == EOS_EResult::EOS_Success)
-		{
-			ON_SCOPE_EXIT
+			EOS_EResult Result = EOS_Auth_CopyUserAuthToken(AuthHandle, &CopyUserAuthTokenOptions, AccountInfoEOS->EpicAccountId, &AuthToken);
+			if (Result == EOS_EResult::EOS_Success)
 			{
-				EOS_Auth_Token_Release(AuthToken);
-			};
-
-			FExternalServerAuthTicket ExternalServerAuthTicket;
-			ExternalServerAuthTicket.Type = ExternalLoginType::Epic;
-			ExternalServerAuthTicket.Data = UTF8_TO_TCHAR(AuthToken->AccessToken);
-			InAsyncOp.SetResult(FAuthQueryExternalServerAuthTicket::Result{MoveTemp(ExternalServerAuthTicket)});
-		}
-		else
-		{
-			InAsyncOp.SetError(Errors::FromEOSResult(Result));
-		}
-	})
-	.Enqueue(GetSerialQueue());
+				ON_SCOPE_EXIT
+				{
+					EOS_Auth_Token_Release(AuthToken);
+				};
 
+				FExternalServerAuthTicket ExternalServerAuthTicket;
+				ExternalServerAuthTicket.Type = ExternalLoginType::Epic;
+				ExternalServerAuthTicket.Data = UTF8_TO_TCHAR(AuthToken->AccessToken);
+				InAsyncOp.SetResult(FAuthQueryExternalServerAuthTicket::Result{ MoveTemp(ExternalServerAuthTicket) });
+			}
+			else
+			{
+				InAsyncOp.SetError(Errors::FromEOSResult(Result));
+			}
+		})
+		.Enqueue(GetSerialQueue());
+	}
 	return Op->GetHandle();
 }
 
 TOnlineAsyncOpHandle<FAuthQueryExternalAuthToken> FAuthEOS::QueryExternalAuthToken(FAuthQueryExternalAuthToken::Params&& Params)
 {
-	TOnlineAsyncOpRef<FAuthQueryExternalAuthToken> Op = GetOp<FAuthQueryExternalAuthToken>(MoveTemp(Params));
-
-	Op->Then([this](TOnlineAsyncOp<FAuthQueryExternalAuthToken>& InAsyncOp)
+	TOnlineAsyncOpRef<FAuthQueryExternalAuthToken> Op = GetJoinableOp<FAuthQueryExternalAuthToken>(MoveTemp(Params));
+	if (!Op->IsReady())
 	{
-		const FAuthQueryExternalAuthToken::Params& Params = InAsyncOp.GetParams();
-		TSharedPtr<FAccountInfoEOS> AccountInfoEOS = AccountInfoRegistryEOS.Find(Params.LocalAccountId);
-		if (!AccountInfoEOS)
+		Op->Then([this](TOnlineAsyncOp<FAuthQueryExternalAuthToken>& InAsyncOp)
 		{
-			InAsyncOp.SetError(Errors::InvalidParams());
-			return;
-		}
+			const FAuthQueryExternalAuthToken::Params& Params = InAsyncOp.GetParams();
+			TSharedPtr<FAccountInfoEOS> AccountInfoEOS = AccountInfoRegistryEOS.Find(Params.LocalAccountId);
+			if (!AccountInfoEOS)
+			{
+				InAsyncOp.SetError(Errors::InvalidParams());
+				return;
+			}
 
-		TDefaultErrorResult<FAuthGetExternalAuthTokenImpl> AuthTokenResult = GetExternalAuthTokenImpl(FAuthGetExternalAuthTokenImpl::Params{ AccountInfoEOS->EpicAccountId });
-		if (AuthTokenResult.IsError())
-		{
-			UE_LOG(LogOnlineServices, Warning, TEXT("[FAuthEOS::QueryExternalAuthToken] Failure: GetExternalAuthTokenImpl %s"), *AuthTokenResult.GetErrorValue().GetLogString());
-			InAsyncOp.SetError(Errors::Unknown(MoveTemp(AuthTokenResult.GetErrorValue())));
-			return;
-		}
+			TDefaultErrorResult<FAuthGetExternalAuthTokenImpl> AuthTokenResult = GetExternalAuthTokenImpl(FAuthGetExternalAuthTokenImpl::Params{ AccountInfoEOS->EpicAccountId });
+			if (AuthTokenResult.IsError())
+			{
+				UE_LOG(LogOnlineServices, Warning, TEXT("[FAuthEOS::QueryExternalAuthToken] Failure: GetExternalAuthTokenImpl %s"), *AuthTokenResult.GetErrorValue().GetLogString());
+				InAsyncOp.SetError(Errors::Unknown(MoveTemp(AuthTokenResult.GetErrorValue())));
+				return;
+			}
 
-		InAsyncOp.SetResult(FAuthQueryExternalAuthToken::Result{MoveTemp(AuthTokenResult.GetOkValue().Token)});
-	})
-	.Enqueue(GetSerialQueue());
+			InAsyncOp.SetResult(FAuthQueryExternalAuthToken::Result{ MoveTemp(AuthTokenResult.GetOkValue().Token) });
+		})
+		.Enqueue(GetSerialQueue());
+	}
 
 	return Op->GetHandle();
 }
