- Use joinable ops in EOS Commerce

diff --git a/Engine/Plugins/Online/Experimental/OnlineServicesEOS/Source/Private/Online/CommerceEOS.cpp b/Engine/Plugins/Online/Experimental/OnlineServicesEOS/Source/Private/Online/CommerceEOS.cpp
--- a/Engine/Plugins/Online/Experimental/OnlineServicesEOS/Source/Private/Online/CommerceEOS.cpp
+++ b/Engine/Plugins/Online/Experimental/OnlineServicesEOS/Source/Private/Online/CommerceEOS.cpp
@@ -55,89 +55,94 @@ void FCommerceEOS::EOSEntitlementToOssEntitlement(FEntitlement& OutEntitlement,
 
 TOnlineAsyncOpHandle<FCommerceQueryOffers> FCommerceEOS::QueryOffers(FCommerceQueryOffers::Params&& Params)
 {
-	TOnlineAsyncOpRef<FCommerceQueryOffers> Op = GetOp<FCommerceQueryOffers>(MoveTemp(Params));
-
-	Op->Then([this](TOnlineAsyncOp<FCommerceQueryOffers>& Op, TPromise<const EOS_Ecom_QueryOffersCallbackInfo*>&& Promise)
+	TOnlineAsyncOpRef<FCommerceQueryOffers> Op = GetJoinableOp<FCommerceQueryOffers>(MoveTemp(Params));
+	if (!Op->IsReady())
 	{
-		const FCommerceQueryOffers::Params& Params = Op.GetParams();
-		if (!Services.Get<FAuthEOS>()->IsLoggedIn(Params.LocalAccountId))
+		Op->Then([this](TOnlineAsyncOp<FCommerceQueryOffers>& Op)
 		{
-			Op.SetError(Errors::NotLoggedIn());
-			return;
-		}
-		EOS_EpicAccountId LocalUserEasId = GetEpicAccountId(Params.LocalAccountId);
-		if (!EOS_EpicAccountId_IsValid(LocalUserEasId))
+			const FCommerceQueryOffers::Params& Params = Op.GetParams();
+			if (!Services.Get<FAuthEOS>()->IsLoggedIn(Params.LocalAccountId))
+			{
+				Op.SetError(Errors::NotLoggedIn());
+				return;
+			}
+			EOS_EpicAccountId LocalUserEasId = GetEpicAccountId(Params.LocalAccountId);
+			if (!EOS_EpicAccountId_IsValid(LocalUserEasId))
+			{
+				Op.SetError(Errors::NotLoggedIn());
+				return;
+			}
+		})
+		.Then([this](TOnlineAsyncOp<FCommerceQueryOffers>& Op, TPromise<const EOS_Ecom_QueryOffersCallbackInfo*>&& Promise)
 		{
-			Op.SetError(Errors::NotLoggedIn());
-			return;
-		}
-		
-		EOS_Ecom_QueryOffersOptions Options = { };
-		Options.ApiVersion = EOS_ECOM_QUERYOFFERS_API_LATEST;
-		Options.LocalUserId = LocalUserEasId;
-		EOS_Async(EOS_Ecom_QueryOffers, EcomHandle, Options, MoveTemp(Promise));
-	})
-	.Then([this](TOnlineAsyncOp<FCommerceQueryOffers>& Op, const EOS_Ecom_QueryOffersCallbackInfo* Data)
-	{
-		EOS_EResult Result = Data->ResultCode;
-		if (Result != EOS_EResult::EOS_Success)
+			const FCommerceQueryOffers::Params& Params = Op.GetParams();
+			EOS_Ecom_QueryOffersOptions Options = { };
+			Options.ApiVersion = EOS_ECOM_QUERYOFFERS_API_LATEST;
+			Options.LocalUserId = GetEpicAccountIdChecked(Params.LocalAccountId);
+			EOS_Async(EOS_Ecom_QueryOffers, EcomHandle, Options, MoveTemp(Promise));
+		})
+		.Then([this](TOnlineAsyncOp<FCommerceQueryOffers>& Op, const EOS_Ecom_QueryOffersCallbackInfo* Data)
 		{
-			Op.SetError(Errors::FromEOSResult(Result));
-			return;
-		}
+			const EOS_EResult Result = Data->ResultCode;
+			if (Result != EOS_EResult::EOS_Success)
+			{
+				Op.SetError(Errors::FromEOSResult(Result));
+				return;
+			}
 
-		EOS_Ecom_GetOfferCountOptions CountOptions = { };
-		CountOptions.ApiVersion = EOS_ECOM_GETOFFERCOUNT_API_LATEST;
-		CountOptions.LocalUserId = Data->LocalUserId;
-		uint32 OfferCount = EOS_Ecom_GetOfferCount(EcomHandle, &CountOptions);
+			EOS_Ecom_GetOfferCountOptions CountOptions = { };
+			CountOptions.ApiVersion = EOS_ECOM_GETOFFERCOUNT_API_LATEST;
+			CountOptions.LocalUserId = Data->LocalUserId;
+			uint32 OfferCount = EOS_Ecom_GetOfferCount(EcomHandle, &CountOptions);
 
-		TArray<FOffer>& Offers = CachedOffers.FindOrAdd(Op.GetParams().LocalAccountId, TArray<FOffer>());
-		Offers.Empty(OfferCount);
+			TArray<FOffer>& Offers = CachedOffers.FindOrAdd(Op.GetParams().LocalAccountId, TArray<FOffer>());
+			Offers.Empty(OfferCount);
 
-		EOS_Ecom_CopyOfferByIndexOptions OfferOptions = { };
-		OfferOptions.ApiVersion = EOS_ECOM_COPYOFFERBYINDEX_API_LATEST;
-		OfferOptions.LocalUserId = Data->LocalUserId;
-		// Iterate and parse the offer list
-		for (uint32 OfferIndex = 0; OfferIndex < OfferCount; OfferIndex++)
-		{
-			EOS_Ecom_CatalogOffer* EosOffer = nullptr;
-			OfferOptions.OfferIndex = OfferIndex;
-			EOS_EResult OfferResult = EOS_Ecom_CopyOfferByIndex(EcomHandle, &OfferOptions, &EosOffer);
-			if (OfferResult != EOS_EResult::EOS_Success)
+			EOS_Ecom_CopyOfferByIndexOptions OfferOptions = { };
+			OfferOptions.ApiVersion = EOS_ECOM_COPYOFFERBYINDEX_API_LATEST;
+			OfferOptions.LocalUserId = Data->LocalUserId;
+			// Iterate and parse the offer list
+			for (uint32 OfferIndex = 0; OfferIndex < OfferCount; OfferIndex++)
 			{
-				continue;
+				EOS_Ecom_CatalogOffer* EosOffer = nullptr;
+				OfferOptions.OfferIndex = OfferIndex;
+				const EOS_EResult OfferResult = EOS_Ecom_CopyOfferByIndex(EcomHandle, &OfferOptions, &EosOffer);
+				if (OfferResult != EOS_EResult::EOS_Success)
+				{
+					continue;
+				}
+				FOffer& Offer = Offers.Emplace_GetRef();
+				EOSOfferToOssOffer(Offer, EosOffer);
+				EOS_Ecom_CatalogOffer_Release(EosOffer);
 			}
-			FOffer& Offer = Offers.Emplace_GetRef();
-			EOSOfferToOssOffer(Offer, EosOffer);
-			EOS_Ecom_CatalogOffer_Release(EosOffer);
-		}
-
-		Op.SetResult({});
-	})
-	.Enqueue(GetSerialQueue());
 
+			Op.SetResult({});
+		})
+		.Enqueue(GetSerialQueue());
+	}
 	return Op->GetHandle();
 }
 
 TOnlineAsyncOpHandle<FCommerceQueryOffersById> FCommerceEOS::QueryOffersById(FCommerceQueryOffersById::Params&& Params)
 {
-	TOnlineAsyncOpRef<FCommerceQueryOffersById> Op = GetOp<FCommerceQueryOffersById>(MoveTemp(Params));
-
-	// EOS doesn't support anything ID-specific, just use the generic QueryOffers call
-	QueryOffers({Op->GetParams().LocalAccountId})
-	.OnComplete([this, Op](const TOnlineResult<FCommerceQueryOffers>& Result)
+	TOnlineAsyncOpRef<FCommerceQueryOffersById> Op = GetJoinableOp<FCommerceQueryOffersById>(MoveTemp(Params));
+	if (!Op->IsReady())
 	{
-		if(Result.IsOk())
+		// EOS doesn't support anything ID-specific, just use the generic QueryOffers call
+		QueryOffers({ Op->GetParams().LocalAccountId })
+			.OnComplete([Op](const TOnlineResult<FCommerceQueryOffers>& Result)
 		{
-			Op->SetResult({});
-		}
-		else
-		{
-			FOnlineError Error = Result.GetErrorValue();
-			Op->SetError(MoveTemp(Error));
-		}
-	});	
-
+			if (Result.IsOk())
+			{
+				Op->SetResult({});
+			}
+			else
+			{
+				FOnlineError Error = Result.GetErrorValue();
+				Op->SetError(MoveTemp(Error));
+			}
+		});
+	}
 	return Op->GetHandle();
 }
 
@@ -235,146 +240,149 @@ TOnlineAsyncOpHandle<FCommerceCheckout> FCommerceEOS::Checkout(FCommerceCheckout
 
 TOnlineAsyncOpHandle<FCommerceQueryTransactionEntitlements> FCommerceEOS::QueryTransactionEntitlements(FCommerceQueryTransactionEntitlements::Params&& Params)
 {
-	TOnlineAsyncOpRef<FCommerceQueryTransactionEntitlements> Op = GetOp<FCommerceQueryTransactionEntitlements>(MoveTemp(Params));
-
-	Op->Then([this](TOnlineAsyncOp<FCommerceQueryTransactionEntitlements>& Op)
+	TOnlineAsyncOpRef<FCommerceQueryTransactionEntitlements> Op = GetJoinableOp<FCommerceQueryTransactionEntitlements>(MoveTemp(Params));
+	if (!Op->IsReady())
 	{
-		const FCommerceQueryTransactionEntitlements::Params& Params = Op.GetParams();
-		if (!Services.Get<FAuthEOS>()->IsLoggedIn(Params.LocalAccountId))
+		Op->Then([this](TOnlineAsyncOp<FCommerceQueryTransactionEntitlements>& Op)
 		{
-			Op.SetError(Errors::NotLoggedIn());
-			return;
-		}
-		EOS_EpicAccountId LocalUserEasId = GetEpicAccountId(Params.LocalAccountId);
-		if (!EOS_EpicAccountId_IsValid(LocalUserEasId))
-		{
-			Op.SetError(Errors::NotLoggedIn());
-			return;
-		}
-
-
-		EOS_Ecom_HTransaction OutTransaction;
-		FTCHARToUTF8 Utf8TransactionId(*Params.TransactionId);
-		EOS_Ecom_CopyTransactionByIdOptions CopyTransactionOptions = { };
-		CopyTransactionOptions.ApiVersion = EOS_ECOM_COPYTRANSACTIONBYID_API_LATEST;
-		CopyTransactionOptions.LocalUserId = LocalUserEasId;
-		CopyTransactionOptions.TransactionId = Utf8TransactionId.Get();
+			const FCommerceQueryTransactionEntitlements::Params& Params = Op.GetParams();
+			if (!Services.Get<FAuthEOS>()->IsLoggedIn(Params.LocalAccountId))
+			{
+				Op.SetError(Errors::NotLoggedIn());
+				return;
+			}
+			EOS_EpicAccountId LocalUserEasId = GetEpicAccountId(Params.LocalAccountId);
+			if (!EOS_EpicAccountId_IsValid(LocalUserEasId))
+			{
+				Op.SetError(Errors::NotLoggedIn());
+				return;
+			}
 
-		EOS_EResult CopyTransactionResult = EOS_Ecom_CopyTransactionById(EcomHandle, &CopyTransactionOptions, &OutTransaction);
-		if (CopyTransactionResult != EOS_EResult::EOS_Success)
-		{
-			UE_LOG(LogTemp, Error, TEXT("EOS_Ecom_CopyTransactionById: failed with error (%s)"), ANSI_TO_TCHAR(EOS_EResult_ToString(CopyTransactionResult)));
-			Op.SetError(Errors::FromEOSResult(CopyTransactionResult));
-			return;
-		}
+			EOS_Ecom_HTransaction OutTransaction;
+			FTCHARToUTF8 Utf8TransactionId(*Params.TransactionId);
+			EOS_Ecom_CopyTransactionByIdOptions CopyTransactionOptions = { };
+			CopyTransactionOptions.ApiVersion = EOS_ECOM_COPYTRANSACTIONBYID_API_LATEST;
+			CopyTransactionOptions.LocalUserId = LocalUserEasId;
+			CopyTransactionOptions.TransactionId = Utf8TransactionId.Get();
 
-		EOS_Ecom_Transaction_GetEntitlementsCountOptions EntitlementCountOptions = { };
-		EntitlementCountOptions.ApiVersion = EOS_ECOM_TRANSACTION_GETENTITLEMENTSCOUNT_API_LATEST;
-		int32 NumTransactionsFound = EOS_Ecom_Transaction_GetEntitlementsCount(OutTransaction, &EntitlementCountOptions);
+			EOS_EResult CopyTransactionResult = EOS_Ecom_CopyTransactionById(EcomHandle, &CopyTransactionOptions, &OutTransaction);
+			if (CopyTransactionResult != EOS_EResult::EOS_Success)
+			{
+				UE_LOG(LogTemp, Error, TEXT("EOS_Ecom_CopyTransactionById: failed with error (%s)"), ANSI_TO_TCHAR(EOS_EResult_ToString(CopyTransactionResult)));
+				Op.SetError(Errors::FromEOSResult(CopyTransactionResult));
+				return;
+			}
 
-		if(NumTransactionsFound > 0)
-		{
-			FCommerceQueryTransactionEntitlements::Result OutResult;
+			EOS_Ecom_Transaction_GetEntitlementsCountOptions EntitlementCountOptions = { };
+			EntitlementCountOptions.ApiVersion = EOS_ECOM_TRANSACTION_GETENTITLEMENTSCOUNT_API_LATEST;
+			int32 NumTransactionsFound = EOS_Ecom_Transaction_GetEntitlementsCount(OutTransaction, &EntitlementCountOptions);
 
-			for(int32 i = 0; i < NumTransactionsFound; i++)
+			if (NumTransactionsFound > 0)
 			{
-				EOS_Ecom_Entitlement* EosEntitlement;
-				EOS_Ecom_Transaction_CopyEntitlementByIndexOptions CopyTransactionEntitlementOptions = {};
-				CopyTransactionEntitlementOptions.ApiVersion = EOS_ECOM_TRANSACTION_COPYENTITLEMENTBYINDEX_API_LATEST;
-				CopyTransactionEntitlementOptions.EntitlementIndex = i;
+				FCommerceQueryTransactionEntitlements::Result OutResult;
 
-				EOS_EResult Result = EOS_Ecom_Transaction_CopyEntitlementByIndex(OutTransaction, &CopyTransactionEntitlementOptions, &EosEntitlement);
-				if(Result == EOS_EResult::EOS_Success)
+				for (int32 i = 0; i < NumTransactionsFound; i++)
 				{
-					FEntitlement& Entitlement = OutResult.Entitlements.Emplace_GetRef();
-					EOSEntitlementToOssEntitlement(Entitlement, EosEntitlement);
-					EOS_Ecom_Entitlement_Release(EosEntitlement);
+					EOS_Ecom_Entitlement* EosEntitlement;
+					EOS_Ecom_Transaction_CopyEntitlementByIndexOptions CopyTransactionEntitlementOptions = {};
+					CopyTransactionEntitlementOptions.ApiVersion = EOS_ECOM_TRANSACTION_COPYENTITLEMENTBYINDEX_API_LATEST;
+					CopyTransactionEntitlementOptions.EntitlementIndex = i;
+
+					EOS_EResult Result = EOS_Ecom_Transaction_CopyEntitlementByIndex(OutTransaction, &CopyTransactionEntitlementOptions, &EosEntitlement);
+					if (Result == EOS_EResult::EOS_Success)
+					{
+						FEntitlement& Entitlement = OutResult.Entitlements.Emplace_GetRef();
+						EOSEntitlementToOssEntitlement(Entitlement, EosEntitlement);
+						EOS_Ecom_Entitlement_Release(EosEntitlement);
+					}
 				}
-			}
 
+				Op.SetResult(MoveTemp(OutResult));
+			}
+			else
+			{
+				Op.SetError(Errors::NotFound());
+			}
 			EOS_Ecom_Transaction_Release(OutTransaction);
-			Op.SetResult(MoveTemp(OutResult));
-		}
-		else
-		{
-			EOS_Ecom_Transaction_Release(OutTransaction);
-			Op.SetError(Errors::NotFound());
-		}
-	})
-	.Enqueue(GetSerialQueue());
-
+		})
+		.Enqueue(GetSerialQueue());
+	}
 	return Op->GetHandle();
 }
 
 TOnlineAsyncOpHandle<FCommerceQueryEntitlements> FCommerceEOS::QueryEntitlements(FCommerceQueryEntitlements::Params&& Params)
 {
-	TOnlineAsyncOpRef<FCommerceQueryEntitlements> Op = GetOp<FCommerceQueryEntitlements>(MoveTemp(Params));
-
-	Op->Then([this](TOnlineAsyncOp<FCommerceQueryEntitlements>& Op, TPromise<const EOS_Ecom_QueryEntitlementsCallbackInfo*>&& Promise)
+	TOnlineAsyncOpRef<FCommerceQueryEntitlements> Op = GetJoinableOp<FCommerceQueryEntitlements>(MoveTemp(Params));
+	if (!Op->IsReady())
 	{
-		const FCommerceQueryEntitlements::Params& Params = Op.GetParams();
-		if (!Services.Get<FAuthEOS>()->IsLoggedIn(Params.LocalAccountId))
+		Op->Then([this](TOnlineAsyncOp<FCommerceQueryEntitlements>& Op)
 		{
-			Op.SetError(Errors::NotLoggedIn());
-			return;
-		}
-		EOS_EpicAccountId LocalUserEasId = GetEpicAccountId(Params.LocalAccountId);
-		if (!EOS_EpicAccountId_IsValid(LocalUserEasId))
+			const FCommerceQueryEntitlements::Params& Params = Op.GetParams();
+			if (!Services.Get<FAuthEOS>()->IsLoggedIn(Params.LocalAccountId))
+			{
+				Op.SetError(Errors::NotLoggedIn());
+				return;
+			}
+			EOS_EpicAccountId LocalUserEasId = GetEpicAccountId(Params.LocalAccountId);
+			if (!EOS_EpicAccountId_IsValid(LocalUserEasId))
+			{
+				Op.SetError(Errors::NotLoggedIn());
+				return;
+			}
+		})
+		.Then([this](TOnlineAsyncOp<FCommerceQueryEntitlements>& Op, TPromise<const EOS_Ecom_QueryEntitlementsCallbackInfo*>&& Promise)
 		{
-			Op.SetError(Errors::NotLoggedIn());
-			return;
-		}
-
-		EOS_Ecom_QueryEntitlementsOptions Options = { };
-		Options.ApiVersion = EOS_ECOM_QUERYENTITLEMENTS_API_LATEST;
-		Options.LocalUserId = LocalUserEasId;
-		Options.bIncludeRedeemed = EOS_TRUE;
-
-		EOS_Async(EOS_Ecom_QueryEntitlements, EcomHandle, Options, MoveTemp(Promise));
-	})
-	.Then([this](TOnlineAsyncOp<FCommerceQueryEntitlements>& Op, const EOS_Ecom_QueryEntitlementsCallbackInfo* Data)
-	{
-		EOS_EResult Result = Data->ResultCode;
-		if (Result != EOS_EResult::EOS_Success)
+			const FCommerceQueryEntitlements::Params& Params = Op.GetParams();
+			EOS_Ecom_QueryEntitlementsOptions Options = { };
+			Options.ApiVersion = EOS_ECOM_QUERYENTITLEMENTS_API_LATEST;
+			Options.LocalUserId = GetEpicAccountIdChecked(Params.LocalAccountId);
+			Options.bIncludeRedeemed = EOS_TRUE;
+			EOS_Async(EOS_Ecom_QueryEntitlements, EcomHandle, Options, MoveTemp(Promise));
+		})
+		.Then([this](TOnlineAsyncOp<FCommerceQueryEntitlements>& Op, const EOS_Ecom_QueryEntitlementsCallbackInfo* Data)
 		{
-			UE_LOG(LogTemp, Error, TEXT("EOS_Ecom_QueryEntitlements: failed with error (%s)"), ANSI_TO_TCHAR(EOS_EResult_ToString(Data->ResultCode)));
-			Op.SetError(Errors::FromEOSResult(Result));
-			return;
-		}
-
-		EOS_Ecom_GetEntitlementsCountOptions CountOptions = { };
-		CountOptions.ApiVersion = EOS_ECOM_GETENTITLEMENTSCOUNT_API_LATEST;
-		CountOptions.LocalUserId = Data->LocalUserId;
-		uint32 Count = EOS_Ecom_GetEntitlementsCount(EcomHandle, &CountOptions);
-		TArray<FEntitlement>& Entitlements = CachedEntitlements.FindOrAdd(Op.GetParams().LocalAccountId, TArray<FEntitlement>());
-		Entitlements.Reset(Count);
+			const EOS_EResult Result = Data->ResultCode;
+			if (Result != EOS_EResult::EOS_Success)
+			{
+				UE_LOG(LogTemp, Error, TEXT("EOS_Ecom_QueryEntitlements: failed with error (%s)"), ANSI_TO_TCHAR(EOS_EResult_ToString(Data->ResultCode)));
+				Op.SetError(Errors::FromEOSResult(Result));
+				return;
+			}
 
-		EOS_Ecom_CopyEntitlementByIndexOptions CopyOptions = { };
-		CopyOptions.ApiVersion = EOS_ECOM_COPYENTITLEMENTBYINDEX_API_LATEST;
-		CopyOptions.LocalUserId = Data->LocalUserId;
+			EOS_Ecom_GetEntitlementsCountOptions CountOptions = { };
+			CountOptions.ApiVersion = EOS_ECOM_GETENTITLEMENTSCOUNT_API_LATEST;
+			CountOptions.LocalUserId = Data->LocalUserId;
+			uint32 Count = EOS_Ecom_GetEntitlementsCount(EcomHandle, &CountOptions);
+			TArray<FEntitlement>& Entitlements = CachedEntitlements.FindOrAdd(Op.GetParams().LocalAccountId, TArray<FEntitlement>());
+			Entitlements.Reset(Count);
 
-		for (uint32 Index = 0; Index < Count; Index++)
-		{
-			CopyOptions.EntitlementIndex = Index;
+			EOS_Ecom_CopyEntitlementByIndexOptions CopyOptions = { };
+			CopyOptions.ApiVersion = EOS_ECOM_COPYENTITLEMENTBYINDEX_API_LATEST;
+			CopyOptions.LocalUserId = Data->LocalUserId;
 
-			EOS_Ecom_Entitlement* EosEntitlement = nullptr;
-			EOS_EResult CopyResult = EOS_Ecom_CopyEntitlementByIndex(EcomHandle, &CopyOptions, &EosEntitlement);
-			if (CopyResult != EOS_EResult::EOS_Success && CopyResult != EOS_EResult::EOS_Ecom_EntitlementStale)
+			for (uint32 Index = 0; Index < Count; Index++)
 			{
-				UE_LOG(LogTemp, Error, TEXT("EOS_Ecom_CopyEntitlementByIndex: failed with error (%s) (proceeding with operation)"), ANSI_TO_TCHAR(EOS_EResult_ToString(CopyResult)));
-				continue;
-			}
+				CopyOptions.EntitlementIndex = Index;
 
-			// Parse the entitlement into the receipt format
-			FEntitlement Entitlement = Entitlements.Emplace_GetRef();
-			EOSEntitlementToOssEntitlement(Entitlement, EosEntitlement);
+				EOS_Ecom_Entitlement* EosEntitlement = nullptr;
+				EOS_EResult CopyResult = EOS_Ecom_CopyEntitlementByIndex(EcomHandle, &CopyOptions, &EosEntitlement);
+				if (CopyResult != EOS_EResult::EOS_Success && CopyResult != EOS_EResult::EOS_Ecom_EntitlementStale)
+				{
+					UE_LOG(LogTemp, Error, TEXT("EOS_Ecom_CopyEntitlementByIndex: failed with error (%s) (proceeding with operation)"), ANSI_TO_TCHAR(EOS_EResult_ToString(CopyResult)));
+					continue;
+				}
 
-			EOS_Ecom_Entitlement_Release(EosEntitlement);
-		}
+				// Parse the entitlement into the receipt format
+				FEntitlement Entitlement = Entitlements.Emplace_GetRef();
+				EOSEntitlementToOssEntitlement(Entitlement, EosEntitlement);
 
-		Op.SetResult({});
-	})
-	.Enqueue(GetSerialQueue());
+				EOS_Ecom_Entitlement_Release(EosEntitlement);
+			}
+
+			Op.SetResult({});
+		})
+		.Enqueue(GetSerialQueue());
+	}
 
 	return Op->GetHandle();
 }
@@ -435,50 +443,53 @@ TOnlineAsyncOpHandle<FCommerceRedeemEntitlement> FCommerceEOS::RedeemEntitlement
 
 TOnlineAsyncOpHandle<FCommerceRetrieveS2SToken> FCommerceEOS::RetrieveS2SToken(FCommerceRetrieveS2SToken::Params&& Params)
 {
-	TOnlineAsyncOpRef<FCommerceRetrieveS2SToken> Op = GetOp<FCommerceRetrieveS2SToken>(MoveTemp(Params));
-
-	Op->Then([this](TOnlineAsyncOp<FCommerceRetrieveS2SToken>& Op, TPromise<const EOS_Ecom_QueryOwnershipTokenCallbackInfo*>&& Promise)
+	TOnlineAsyncOpRef<FCommerceRetrieveS2SToken> Op = GetJoinableOp<FCommerceRetrieveS2SToken>(MoveTemp(Params));
+	if (!Op->IsReady())
 	{
-		const FCommerceRetrieveS2SToken::Params& Params = Op.GetParams();
-		if (!Services.Get<FAuthEOS>()->IsLoggedIn(Params.LocalAccountId))
+		Op->Then([this](TOnlineAsyncOp<FCommerceRetrieveS2SToken>& Op)
 		{
-			Op.SetError(Errors::NotLoggedIn());
-			return;
-		}
-		EOS_EpicAccountId LocalUserEasId = GetEpicAccountId(Params.LocalAccountId);
-		if (!EOS_EpicAccountId_IsValid(LocalUserEasId))
-		{
-			Op.SetError(Errors::NotLoggedIn());
-			return;
-		}
-
-		EOS_Ecom_QueryOwnershipTokenOptions Options = {};
-		Options.ApiVersion = EOS_ECOM_QUERYOWNERSHIPTOKEN_API_LATEST;
-		Options.LocalUserId = LocalUserEasId;
-
-		EOS_Async(EOS_Ecom_QueryOwnershipToken, EcomHandle, Options, MoveTemp(Promise));
-	})
-	.Then([this](TOnlineAsyncOp<FCommerceRetrieveS2SToken>& Op, const EOS_Ecom_QueryOwnershipTokenCallbackInfo* Data)
-	{
-		EOS_EResult Result = Data->ResultCode;
-		if (Result != EOS_EResult::EOS_Success)
+			const FCommerceRetrieveS2SToken::Params& Params = Op.GetParams();
+			if (!Services.Get<FAuthEOS>()->IsLoggedIn(Params.LocalAccountId))
+			{
+				Op.SetError(Errors::NotLoggedIn());
+				return;
+			}
+			EOS_EpicAccountId LocalUserEasId = GetEpicAccountId(Params.LocalAccountId);
+			if (!EOS_EpicAccountId_IsValid(LocalUserEasId))
+			{
+				Op.SetError(Errors::NotLoggedIn());
+				return;
+			}
+		})
+		.Then([this](TOnlineAsyncOp<FCommerceRetrieveS2SToken>& Op, TPromise<const EOS_Ecom_QueryOwnershipTokenCallbackInfo*>&& Promise)
 		{
-			UE_LOG(LogTemp, Error, TEXT("EOS_Ecom_QueryOwnershipToken: failed with error (%s)"), ANSI_TO_TCHAR(EOS_EResult_ToString(Data->ResultCode)));
-			Op.SetError(Errors::FromEOSResult(Result));
-			return;
-		}
-
-		size_t DataSize = strlen(Data->OwnershipToken);
-		FCommerceRetrieveS2SToken::Result ResultValue;
-		ResultValue.Token.AddZeroed(DataSize);
-		for(size_t i = 0; i < DataSize; i++)
+			const FCommerceRetrieveS2SToken::Params& Params = Op.GetParams();
+			EOS_Ecom_QueryOwnershipTokenOptions Options = {};
+			Options.ApiVersion = EOS_ECOM_QUERYOWNERSHIPTOKEN_API_LATEST;
+			Options.LocalUserId = GetEpicAccountIdChecked(Params.LocalAccountId);
+			EOS_Async(EOS_Ecom_QueryOwnershipToken, EcomHandle, Options, MoveTemp(Promise));
+		})
+		.Then([this](TOnlineAsyncOp<FCommerceRetrieveS2SToken>& Op, const EOS_Ecom_QueryOwnershipTokenCallbackInfo* Data)
 		{
-			ResultValue.Token[i] = (uint8) Data->OwnershipToken[i];
-		}
-		Op.SetResult(MoveTemp(ResultValue));
-	})
-	.Enqueue(GetSerialQueue());
+			EOS_EResult Result = Data->ResultCode;
+			if (Result != EOS_EResult::EOS_Success)
+			{
+				UE_LOG(LogTemp, Error, TEXT("EOS_Ecom_QueryOwnershipToken: failed with error (%s)"), ANSI_TO_TCHAR(EOS_EResult_ToString(Data->ResultCode)));
+				Op.SetError(Errors::FromEOSResult(Result));
+				return;
+			}
 
+			size_t DataSize = strlen(Data->OwnershipToken);
+			FCommerceRetrieveS2SToken::Result ResultValue;
+			ResultValue.Token.AddZeroed(DataSize);
+			for (size_t i = 0; i < DataSize; i++)
+			{
+				ResultValue.Token[i] = (uint8)Data->OwnershipToken[i];
+			}
+			Op.SetResult(MoveTemp(ResultValue));
+		})
+		.Enqueue(GetSerialQueue());
+	}
 	return Op->GetHandle();
 }
 
