- Use joinable ops in EOS UserInfo

diff --git a/Engine/Plugins/Online/Experimental/OnlineServicesEOS/Source/Private/Online/UserInfoEOS.cpp b/Engine/Plugins/Online/Experimental/OnlineServicesEOS/Source/Private/Online/UserInfoEOS.cpp
--- a/Engine/Plugins/Online/Experimental/OnlineServicesEOS/Source/Private/Online/UserInfoEOS.cpp
+++ b/Engine/Plugins/Online/Experimental/OnlineServicesEOS/Source/Private/Online/UserInfoEOS.cpp
@@ -28,58 +28,62 @@ void FUserInfoEOS::Initialize()
 TOnlineAsyncOpHandle<FQueryUserInfo> FUserInfoEOS::QueryUserInfo(FQueryUserInfo::Params&& InParams)
 {
 	TOnlineAsyncOpRef<FQueryUserInfo> Op = GetJoinableOp<FQueryUserInfo>(MoveTemp(InParams));
-	const FQueryUserInfo::Params& Params = Op->GetParams();
-
-	if (Params.AccountIds.IsEmpty())
+	if (!Op->IsReady())
 	{
-		Op->SetError(Errors::InvalidParams());
-		return Op->GetHandle();
-	}
+		const FQueryUserInfo::Params& Params = Op->GetParams();
 
-	for (const FAccountId TargetAccountId : Params.AccountIds)
-	{
-		Op->Then([this, TargetAccountId](TOnlineAsyncOp<FQueryUserInfo>& Op, TPromise<const EOS_UserInfo_QueryUserInfoCallbackInfo*>&& Promise)
+		if (Params.AccountIds.IsEmpty())
 		{
-			const FQueryUserInfo::Params& Params = Op.GetParams();
+			Op->SetError(Errors::InvalidParams());
+			return Op->GetHandle();
+		}
 
-			if (!Services.Get<FAuthEOS>()->IsLoggedIn(Params.LocalAccountId))
+		for (const FAccountId TargetAccountId : Params.AccountIds)
+		{
+			Op->Then([this, TargetAccountId](TOnlineAsyncOp<FQueryUserInfo>& Op, TPromise<const EOS_UserInfo_QueryUserInfoCallbackInfo*>&& Promise)
 			{
-				Op.SetError(Errors::NotLoggedIn());
-				Promise.EmplaceValue();
-				return;
-			}
-
-			const EOS_EpicAccountId TargetUserEasId = GetEpicAccountId(TargetAccountId);
-			if (!EOS_EpicAccountId_IsValid(TargetUserEasId))
+				const FQueryUserInfo::Params& Params = Op.GetParams();
+
+				if (!Services.Get<FAuthEOS>()->IsLoggedIn(Params.LocalAccountId))
+				{
+					Op.SetError(Errors::NotLoggedIn());
+					Promise.EmplaceValue();
+					return;
+				}
+
+				const EOS_EpicAccountId TargetUserEasId = GetEpicAccountId(TargetAccountId);
+				if (!EOS_EpicAccountId_IsValid(TargetUserEasId))
+				{
+					Op.SetError(Errors::InvalidParams());
+					Promise.EmplaceValue();
+					return;
+				}
+
+				EOS_UserInfo_QueryUserInfoOptions QueryUserInfoOptions = {};
+				QueryUserInfoOptions.ApiVersion = EOS_USERINFO_QUERYUSERINFO_API_LATEST;
+				static_assert(EOS_USERINFO_QUERYUSERINFO_API_LATEST == 1, "EOS_UserInfo_QueryUserInfoOptions updated, check new fields");
+				QueryUserInfoOptions.LocalUserId = GetEpicAccountIdChecked(Params.LocalAccountId);
+				QueryUserInfoOptions.TargetUserId = TargetUserEasId;
+
+				EOS_Async(EOS_UserInfo_QueryUserInfo, UserInfoHandle, QueryUserInfoOptions, MoveTemp(Promise));
+			})
+			.Then([this](TOnlineAsyncOp<FQueryUserInfo>& Op, const EOS_UserInfo_QueryUserInfoCallbackInfo* CallbackInfo) mutable
 			{
-				Op.SetError(Errors::InvalidParams());
-				Promise.EmplaceValue();
-				return;
-			}
-
-			EOS_UserInfo_QueryUserInfoOptions QueryUserInfoOptions = {};
-			QueryUserInfoOptions.ApiVersion = EOS_USERINFO_QUERYUSERINFO_API_LATEST;
-			static_assert(EOS_USERINFO_QUERYUSERINFO_API_LATEST == 1, "EOS_UserInfo_QueryUserInfoOptions updated, check new fields");
-			QueryUserInfoOptions.LocalUserId = GetEpicAccountIdChecked(Params.LocalAccountId);
-			QueryUserInfoOptions.TargetUserId = TargetUserEasId;
-
-			EOS_Async(EOS_UserInfo_QueryUserInfo, UserInfoHandle, QueryUserInfoOptions, MoveTemp(Promise));
-		}).Then([this](TOnlineAsyncOp<FQueryUserInfo>& Op, const EOS_UserInfo_QueryUserInfoCallbackInfo* CallbackInfo) mutable
-		{
-			if (CallbackInfo->ResultCode != EOS_EResult::EOS_Success)
+				if (CallbackInfo->ResultCode != EOS_EResult::EOS_Success)
+				{
+					UE_LOG(LogTemp, Warning, TEXT("EOS_UserInfo_QueryUserInfo failed with result=[%s]"), *LexToString(CallbackInfo->ResultCode));
+					Op.SetError(Errors::FromEOSResult(CallbackInfo->ResultCode));
+				}
+			});
+		}
+
+		Op->Then([](TOnlineAsyncOp<FQueryUserInfo>& Op)
 			{
-				UE_LOG(LogTemp, Warning, TEXT("EOS_UserInfo_QueryUserInfo failed with result=[%s]"), *LexToString(CallbackInfo->ResultCode));
-				Op.SetError(Errors::FromEOSResult(CallbackInfo->ResultCode));
-			}
-		});
-	}
+				Op.SetResult({});
+			});
 
-	Op->Then([](TOnlineAsyncOp<FQueryUserInfo>& Op)
-	{
-		Op.SetResult({});
-	});
-	
-	Op->Enqueue(GetSerialQueue());
+		Op->Enqueue(GetSerialQueue());
+	}
 	return Op->GetHandle();
 }
 
