OnlineServices - Social Interface Common
OnlineServices - SocialEOS

diff --git a/Engine/Plugins/Online/Experimental/OnlineServices/Source/OnlineServicesCommon/Private/Online/FriendsCommon.cpp b/Engine/Plugins/Online/Experimental/OnlineServices/Source/OnlineServicesCommon/Private/Online/FriendsCommon.cpp
deleted file mode 100644
diff --git a/Engine/Plugins/Online/Experimental/OnlineServices/Source/OnlineServicesCommon/Private/Online/OnlineServicesCommon.cpp b/Engine/Plugins/Online/Experimental/OnlineServices/Source/OnlineServicesCommon/Private/Online/OnlineServicesCommon.cpp
--- a/Engine/Plugins/Online/Experimental/OnlineServices/Source/OnlineServicesCommon/Private/Online/OnlineServicesCommon.cpp
+++ b/Engine/Plugins/Online/Experimental/OnlineServices/Source/OnlineServicesCommon/Private/Online/OnlineServicesCommon.cpp
@@ -3,7 +3,6 @@
 #include "Online/OnlineServicesCommon.h"
 
 #include "Online/AuthCommon.h"
-#include "Online/FriendsCommon.h"
 
 DEFINE_LOG_CATEGORY(LogOnlineServices);
 
@@ -61,9 +60,9 @@ IUserInfoPtr FOnlineServicesCommon::GetUserInfoInterface()
 	return IUserInfoPtr(AsShared(), Get<IUserInfo>());
 }
 
-IFriendsPtr FOnlineServicesCommon::GetFriendsInterface()
+ISocialPtr FOnlineServicesCommon::GetSocialInterface()
 {
-	return IFriendsPtr(AsShared(), Get<IFriends>());
+	return ISocialPtr(AsShared(), Get<ISocial>());
 }
 
 IPresencePtr FOnlineServicesCommon::GetPresenceInterface()
diff --git a/Engine/Plugins/Online/Experimental/OnlineServices/Source/OnlineServicesCommon/Private/Online/SocialCommon.cpp b/Engine/Plugins/Online/Experimental/OnlineServices/Source/OnlineServicesCommon/Private/Online/SocialCommon.cpp
new file mode 100644
--- /dev/null
+++ b/Engine/Plugins/Online/Experimental/OnlineServices/Source/OnlineServicesCommon/Private/Online/SocialCommon.cpp
@@ -0,0 +1,95 @@
+// Copyright Epic Games, Inc. All Rights Reserved.
+
+#include "Online/SocialCommon.h"
+
+#include "Online/OnlineServicesCommon.h"
+
+namespace UE::Online {
+
+FSocialCommon::FSocialCommon(FOnlineServicesCommon& InServices)
+	: TOnlineComponent(TEXT("Social"), InServices)
+{
+}
+
+void FSocialCommon::RegisterCommands()
+{
+	RegisterCommand(&FSocialCommon::QueryFriends);
+	RegisterCommand(&FSocialCommon::GetFriends);
+	RegisterCommand(&FSocialCommon::SendFriendInvite);
+	RegisterCommand(&FSocialCommon::AcceptFriendInvite);
+	RegisterCommand(&FSocialCommon::RejectFriendInvite);
+	RegisterCommand(&FSocialCommon::QueryBlockedUsers);
+	RegisterCommand(&FSocialCommon::GetBlockedUsers);
+	RegisterCommand(&FSocialCommon::BlockUser);
+}
+
+TOnlineAsyncOpHandle<FQueryFriends> FSocialCommon::QueryFriends(FQueryFriends::Params&& Params)
+{
+	TOnlineAsyncOpRef<FQueryFriends> Operation = GetOp<FQueryFriends>(MoveTemp(Params));
+	Operation->SetError(Errors::NotImplemented());
+	return Operation->GetHandle();
+}
+
+TOnlineResult<FGetFriends> FSocialCommon::GetFriends(FGetFriends::Params&& Params)
+{
+	return TOnlineResult<FGetFriends>(Errors::NotImplemented());
+}
+
+TOnlineAsyncOpHandle<FSendFriendInvite> FSocialCommon::SendFriendInvite(FSendFriendInvite::Params&& Params)
+{
+	TOnlineAsyncOpRef<FSendFriendInvite> Operation = GetOp<FSendFriendInvite>(MoveTemp(Params));
+	Operation->SetError(Errors::NotImplemented());
+	return Operation->GetHandle();
+}
+
+TOnlineAsyncOpHandle<FAcceptFriendInvite> FSocialCommon::AcceptFriendInvite(FAcceptFriendInvite::Params&& Params)
+{
+	TOnlineAsyncOpRef<FAcceptFriendInvite> Operation = GetOp<FAcceptFriendInvite>(MoveTemp(Params));
+	Operation->SetError(Errors::NotImplemented());
+	return Operation->GetHandle();
+}
+
+TOnlineAsyncOpHandle<FRejectFriendInvite> FSocialCommon::RejectFriendInvite(FRejectFriendInvite::Params&& Params)
+{
+	TOnlineAsyncOpRef<FRejectFriendInvite> Operation = GetOp<FRejectFriendInvite>(MoveTemp(Params));
+	Operation->SetError(Errors::NotImplemented());
+	return Operation->GetHandle();
+}
+
+TOnlineEvent<void(const FRelationshipUpdated&)> FSocialCommon::OnRelationshipUpdated()
+{
+	return OnRelationshipUpdatedEvent;
+}
+
+TOnlineAsyncOpHandle<FQueryBlockedUsers> FSocialCommon::QueryBlockedUsers(FQueryBlockedUsers::Params&& Params)
+{
+	TOnlineAsyncOpRef<FQueryBlockedUsers> Operation = GetOp<FQueryBlockedUsers>(MoveTemp(Params));
+	Operation->SetError(Errors::NotImplemented());
+	return Operation->GetHandle();
+}
+
+TOnlineResult<FGetBlockedUsers> FSocialCommon::GetBlockedUsers(FGetBlockedUsers::Params&& Params)
+{
+	return TOnlineResult<FGetBlockedUsers>(Errors::NotImplemented());
+}
+
+TOnlineAsyncOpHandle<FBlockUser> FSocialCommon::BlockUser(FBlockUser::Params&& Params)
+{
+	TOnlineAsyncOpRef<FBlockUser> Operation = GetOp<FBlockUser>(MoveTemp(Params));
+	Operation->SetError(Errors::NotImplemented());
+	return Operation->GetHandle();
+}
+
+void FSocialCommon::BroadcastRelationshipUpdated(FAccountId LocalAccountId, FAccountId RemoteAccountId, ERelationship OldRelationship, ERelationship NewRelationship)
+{
+	FRelationshipUpdated Event;
+	Event.LocalAccountId = LocalAccountId;
+	Event.RemoteAccountId = RemoteAccountId;
+	Event.OldRelationship = OldRelationship;
+	Event.NewRelationship = NewRelationship;
+
+	OnRelationshipUpdatedEvent.Broadcast(Event);
+}
+
+
+/* UE::Online */ }
diff --git a/Engine/Plugins/Online/Experimental/OnlineServices/Source/OnlineServicesCommon/Public/Online/FriendsCommon.h b/Engine/Plugins/Online/Experimental/OnlineServices/Source/OnlineServicesCommon/Public/Online/FriendsCommon.h
deleted file mode 100644
diff --git a/Engine/Plugins/Online/Experimental/OnlineServices/Source/OnlineServicesCommon/Public/Online/OnlineServicesCommon.h b/Engine/Plugins/Online/Experimental/OnlineServices/Source/OnlineServicesCommon/Public/Online/OnlineServicesCommon.h
--- a/Engine/Plugins/Online/Experimental/OnlineServices/Source/OnlineServicesCommon/Public/Online/OnlineServicesCommon.h
+++ b/Engine/Plugins/Online/Experimental/OnlineServices/Source/OnlineServicesCommon/Public/Online/OnlineServicesCommon.h
@@ -39,7 +39,7 @@ public:
 	virtual IAuthPtr GetAuthInterface() override;
 	virtual IUserInfoPtr GetUserInfoInterface() override;
 	virtual ICommercePtr GetCommerceInterface() override;
-	virtual IFriendsPtr GetFriendsInterface() override;
+	virtual ISocialPtr GetSocialInterface() override;
 	virtual IPresencePtr GetPresenceInterface() override;
 	virtual IExternalUIPtr GetExternalUIInterface() override;
 	virtual ILeaderboardsPtr GetLeaderboardsInterface() override;
diff --git a/Engine/Plugins/Online/Experimental/OnlineServices/Source/OnlineServicesCommon/Public/Online/SocialCommon.h b/Engine/Plugins/Online/Experimental/OnlineServices/Source/OnlineServicesCommon/Public/Online/SocialCommon.h
new file mode 100644
--- /dev/null
+++ b/Engine/Plugins/Online/Experimental/OnlineServices/Source/OnlineServicesCommon/Public/Online/SocialCommon.h
@@ -0,0 +1,38 @@
+// Copyright Epic Games, Inc. All Rights Reserved.
+
+#pragma once
+
+#include "Online/Social.h"
+#include "Online/OnlineComponent.h"
+
+namespace UE::Online {
+
+class FOnlineServicesCommon;
+
+class ONLINESERVICESCOMMON_API FSocialCommon : public TOnlineComponent<ISocial>
+{
+public:
+	using Super = ISocial;
+
+	FSocialCommon(FOnlineServicesCommon& InServices);
+
+	// TOnlineComponent
+	virtual void RegisterCommands() override;
+
+	// ISocial
+	virtual TOnlineAsyncOpHandle<FQueryFriends> QueryFriends(FQueryFriends::Params&& Params) override;
+	virtual TOnlineResult<FGetFriends> GetFriends(FGetFriends::Params&& Params) override;
+	virtual TOnlineAsyncOpHandle<FSendFriendInvite> SendFriendInvite(FSendFriendInvite::Params&& Params) override;
+	virtual TOnlineAsyncOpHandle<FAcceptFriendInvite> AcceptFriendInvite(FAcceptFriendInvite::Params&& Params) override;
+	virtual TOnlineAsyncOpHandle<FRejectFriendInvite> RejectFriendInvite(FRejectFriendInvite::Params&& Params) override;
+	virtual TOnlineEvent<void(const FRelationshipUpdated&)> OnRelationshipUpdated() override;
+	virtual TOnlineAsyncOpHandle<FQueryBlockedUsers> QueryBlockedUsers(FQueryBlockedUsers::Params&& Params) override;
+	virtual TOnlineResult<FGetBlockedUsers> GetBlockedUsers(FGetBlockedUsers::Params&& Params) override;
+	virtual TOnlineAsyncOpHandle<FBlockUser> BlockUser(FBlockUser::Params&& Params) override;
+
+protected:
+	void BroadcastRelationshipUpdated(FAccountId LocalAccountId, FAccountId RemoteAccountId, ERelationship OldRelationship, ERelationship NewRelationship);
+	TOnlineEventCallable<void(const FRelationshipUpdated&)> OnRelationshipUpdatedEvent;
+};
+
+/* UE::Online */ }
diff --git a/Engine/Plugins/Online/Experimental/OnlineServices/Source/OnlineServicesInterface/Public/Online/Friends.h b/Engine/Plugins/Online/Experimental/OnlineServices/Source/OnlineServicesInterface/Public/Online/Friends.h
deleted file mode 100644
diff --git a/Engine/Plugins/Online/Experimental/OnlineServices/Source/OnlineServicesInterface/Public/Online/OnlineServices.h b/Engine/Plugins/Online/Experimental/OnlineServices/Source/OnlineServicesInterface/Public/Online/OnlineServices.h
--- a/Engine/Plugins/Online/Experimental/OnlineServices/Source/OnlineServicesInterface/Public/Online/OnlineServices.h
+++ b/Engine/Plugins/Online/Experimental/OnlineServices/Source/OnlineServicesInterface/Public/Online/OnlineServices.h
@@ -17,7 +17,7 @@ using IAchievementsPtr = TSharedPtr<class IAchievements>;
 using IAuthPtr = TSharedPtr<class IAuth>;
 using ICommercePtr = TSharedPtr<class ICommerce>;
 using IUserInfoPtr = TSharedPtr<class IUserInfo>;
-using IFriendsPtr = TSharedPtr<class IFriends>;
+using ISocialPtr = TSharedPtr<class ISocial>;
 using IPresencePtr = TSharedPtr<class IPresence>;
 using IExternalUIPtr = TSharedPtr<class IExternalUI>;
 using ILeaderboardsPtr = TSharedPtr<class ILeaderboards>;
@@ -80,10 +80,12 @@ public:
 	 */
 	virtual IUserInfoPtr GetUserInfoInterface() = 0;
 
+
 	/**
-	 *
+	 * Get the social interface, used to interact with friends lists, blocked user lists, and any other social relationships
+	 * @return social interface implementation, may be null if not implemented for this service
 	 */
-	virtual IFriendsPtr GetFriendsInterface() = 0;
+	virtual ISocialPtr GetSocialInterface() = 0;
 
 	/**
 	 *
diff --git a/Engine/Plugins/Online/Experimental/OnlineServices/Source/OnlineServicesInterface/Public/Online/Social.h b/Engine/Plugins/Online/Experimental/OnlineServices/Source/OnlineServicesInterface/Public/Online/Social.h
new file mode 100644
--- /dev/null
+++ b/Engine/Plugins/Online/Experimental/OnlineServices/Source/OnlineServicesInterface/Public/Online/Social.h
@@ -0,0 +1,387 @@
+// Copyright Epic Games, Inc. All Rights Reserved.
+
+#pragma once
+
+#include "Online/CoreOnline.h"
+#include "Online/OnlineAsyncOpHandle.h"
+#include "Online/OnlineMeta.h"
+
+namespace UE::Online {
+
+/**
+ * Social relationship status
+ */
+enum class ERelationship : uint8
+{
+	/* Friend */
+	Friend,
+	/* Not friend */
+	NotFriend,
+	/* Invite Sent to User */
+	InviteSent,
+	/* Invite Received from User */
+	InviteReceived,
+	/* Local user has Blocked the other user */
+	Blocked
+};
+
+const TCHAR* LexToString(ERelationship Relationship);
+
+/**
+ * Information about a friend
+ */
+struct FFriend
+{
+	/* Account Id of the Friend */
+	FAccountId FriendId;
+	/* Display Name of the Friend */
+	FString DisplayName;
+	/* Local Nickname for the Friend, Not always available */
+	FString Nickname;
+	/* Relationship of the friend */
+	ERelationship Relationship;
+};
+
+struct FQueryFriends
+{
+	static constexpr TCHAR Name[] = TEXT("QueryFriends");
+
+	/** Input struct for Social::QueryFriends */
+	struct Params
+	{
+		/** Account Id of the local user making the request */
+		FAccountId LocalAccountId;
+	};
+
+	/**
+	 * Output struct for Social::QueryFriends
+	 * Obtain the friends list via GetFriends
+	 */
+	struct Result
+	{
+	};
+};
+
+struct FGetFriends
+{
+	static constexpr TCHAR Name[] = TEXT("GetFriends");
+
+	/** Input struct for Social::GetFriends */
+	struct Params
+	{
+		/** Account Id of the local user making the request */
+		FAccountId LocalAccountId;
+	};
+
+	/** Output struct for Social::GetFriends */
+	struct Result
+	{
+		/** Array of friends */
+		TArray<TSharedRef<FFriend>> Friends;
+	};
+};
+
+struct FSendFriendInvite
+{
+	static constexpr TCHAR Name[] = TEXT("SendFriendInvite");
+
+	/** Input struct for Social::SendFriendInvite */
+	struct Params
+	{
+		/** Account Id of the local user making the request */
+		FAccountId LocalAccountId;
+		/** Account Id of user to send a friend request to */
+		FAccountId TargetAccountId;
+	};
+
+	/** Output struct for Social::SendFriendInvite */
+	struct Result
+	{
+	};
+};
+
+struct FAcceptFriendInvite
+{
+	static constexpr TCHAR Name[] = TEXT("AcceptFriendInvite");
+
+	/** Input struct for Social::AcceptFriendInvite */
+	struct Params
+	{
+		/** Account Id of the local user making the request */
+		FAccountId LocalAccountId;
+		/** Account Id of user that sent the invite */
+		FAccountId TargetAccountId;
+	};
+
+	/** Output struct for Social::AcceptFriendInvite */
+	struct Result
+	{
+	};
+};
+
+struct FRejectFriendInvite
+{
+	static constexpr TCHAR Name[] = TEXT("RejectFriendInvite");
+
+	/** Input struct for Social::RejectFriendInvite */
+	struct Params
+	{
+		/** Account Id of the local user making the request */
+		FAccountId LocalAccountId;
+		/** Account Id of user that sent the invite */
+		FAccountId TargetAccountId;
+	};
+
+	/** Output struct for Social::RejectFriendInvite */
+	struct Result
+	{
+	};
+};
+
+/** Struct for RelationshipUpdated event */
+struct FRelationshipUpdated
+{
+	/* Account Id of the Local User whose friends list changed */
+	FAccountId LocalAccountId;
+	/* Account Id of the User whose status changed */
+	FAccountId RemoteAccountId;
+	/* Previous relationship */
+	ERelationship OldRelationship;
+	/*  New relationship */
+	ERelationship NewRelationship;
+};
+
+struct FQueryBlockedUsers
+{
+	static constexpr TCHAR Name[] = TEXT("QueryBlockedUsers");
+
+	/** Input struct for Social::QueryBlockedUsers */
+	struct Params
+	{
+		/** Account Id of the local user making the request */
+		FAccountId LocalAccountId;
+	};
+
+	/**
+	 * Output struct for Social::QueryBlockedUsers
+	 * Obtain the blocked users list via GetBlockedUsers
+	 */
+	struct Result
+	{
+	};
+};
+
+struct FGetBlockedUsers
+{
+	static constexpr TCHAR Name[] = TEXT("GetBlockedUsers");
+
+	/** Input struct for Social::GetBlockedUsers */
+	struct Params
+	{
+		/** Account Id of the local user making the request */
+		FAccountId LocalAccountId;
+	};
+
+	/** Output struct for Social::GetBlockedUsers */
+	struct Result
+	{
+		/** Array of blocked users */
+		TArray<FAccountId> BlockedUsers;
+	};
+};
+
+struct FBlockUser
+{
+	static constexpr TCHAR Name[] = TEXT("BlockUser");
+
+	/** Input struct for Social::BlockUser */
+	struct Params
+	{
+		/** Account Id of the local user making the request */
+		FAccountId LocalAccountId;
+		/** Account Id of user to block */
+		FAccountId TargetAccountId;
+	};
+
+	/** Output struct for Social::BlockUser */
+	struct Result
+	{
+	};
+};
+
+class ISocial
+{
+public:
+	/**
+	 * Query the friends list
+	 * 
+	 * @param Params for the QueryFriends call
+	 * @return AsyncOpHandle
+	 */
+	virtual TOnlineAsyncOpHandle<FQueryFriends> QueryFriends(FQueryFriends::Params&& Params) = 0;
+
+	/**
+	 * Get the contents of a previously queried friends list
+	 *
+	 * @param Params for the GetFriends call
+	 * @return Result
+	 */
+	virtual TOnlineResult<FGetFriends> GetFriends(FGetFriends::Params&& Params) = 0;
+
+	/**
+	 * Send a friend invite
+	 *
+	 * @param Params for the SendFriendInvite call
+	 * @return AsyncOpHandle
+	 */
+	virtual TOnlineAsyncOpHandle<FSendFriendInvite> SendFriendInvite(FSendFriendInvite::Params&& Params) = 0;
+
+	/**
+	 * Accept a friend invite
+	 *
+	 * @param Params for the AcceptFriendInvite call
+	 * @return AsyncOpHandle
+	 */
+	virtual TOnlineAsyncOpHandle<FAcceptFriendInvite> AcceptFriendInvite(FAcceptFriendInvite::Params&& Params) = 0;
+
+	/**
+	 * Reject a friend invite
+	 *
+	 * @param Params for the RejectFriendInvite call
+	 * @return AsyncOpHandle
+	 */
+	virtual TOnlineAsyncOpHandle<FRejectFriendInvite> RejectFriendInvite(FRejectFriendInvite::Params&& Params) = 0;
+
+	/**
+	 * Get the event that is triggered when a friends list is updated
+	 * This typically happens when QueryFriends is called, a friend or block list modifying function is called,
+	 *   or is called in response to an event coming from a backend service
+	 *
+	 * @return Event that can be bound to
+	 */
+	virtual TOnlineEvent<void(const FRelationshipUpdated&)> OnRelationshipUpdated() = 0;
+
+	/**
+	 * Query the blocked users list
+	 *
+	 * @param Params for the QueryBlockedUsers call
+	 * @return AsyncOpHandle
+	 */
+	virtual TOnlineAsyncOpHandle<FQueryBlockedUsers> QueryBlockedUsers(FQueryBlockedUsers::Params&& Params) = 0;
+
+	/**
+	 * Get the contents of a previously queried blocked users list
+	 *
+	 * @param Params for the GetBlockedUsers call
+	 * @return Result
+	 */
+	virtual TOnlineResult<FGetBlockedUsers> GetBlockedUsers(FGetBlockedUsers::Params&& Params) = 0;
+
+	/**
+	 * Block a user
+	 *
+	 * @param Params for the BlockUser call
+	 * @return AsyncOpHandle
+	 */
+	virtual TOnlineAsyncOpHandle<FBlockUser> BlockUser(FBlockUser::Params&& Params) = 0;
+};
+
+inline const TCHAR* LexToString(ERelationship Relationship)
+{
+	switch (Relationship)
+	{
+	case ERelationship::Friend:
+		return TEXT("Friend");
+	case ERelationship::NotFriend:
+		return TEXT("NotFriend");
+	case ERelationship::InviteSent:
+		return TEXT("InviteSent");
+	case ERelationship::InviteReceived:
+		return TEXT("InviteReceived");
+	case ERelationship::Blocked:
+		return TEXT("Blocked");
+	}
+	return TEXT("Invalid");
+}
+
+namespace Meta {
+
+BEGIN_ONLINE_STRUCT_META(FFriend)
+	ONLINE_STRUCT_FIELD(FFriend, FriendId),
+	ONLINE_STRUCT_FIELD(FFriend, DisplayName),
+	ONLINE_STRUCT_FIELD(FFriend, Nickname),
+	ONLINE_STRUCT_FIELD(FFriend, Relationship)
+END_ONLINE_STRUCT_META()
+
+BEGIN_ONLINE_STRUCT_META(FQueryFriends::Params)
+	ONLINE_STRUCT_FIELD(FQueryFriends::Params, LocalAccountId)
+END_ONLINE_STRUCT_META()
+
+BEGIN_ONLINE_STRUCT_META(FQueryFriends::Result)
+END_ONLINE_STRUCT_META()
+
+BEGIN_ONLINE_STRUCT_META(FGetFriends::Params)
+	ONLINE_STRUCT_FIELD(FGetFriends::Params, LocalAccountId)
+END_ONLINE_STRUCT_META()
+
+BEGIN_ONLINE_STRUCT_META(FGetFriends::Result)
+	ONLINE_STRUCT_FIELD(FGetFriends::Result, Friends)
+END_ONLINE_STRUCT_META()
+
+BEGIN_ONLINE_STRUCT_META(FSendFriendInvite::Params)
+	ONLINE_STRUCT_FIELD(FSendFriendInvite::Params, LocalAccountId),
+	ONLINE_STRUCT_FIELD(FSendFriendInvite::Params, TargetAccountId)
+END_ONLINE_STRUCT_META()
+
+BEGIN_ONLINE_STRUCT_META(FSendFriendInvite::Result)
+END_ONLINE_STRUCT_META()
+
+BEGIN_ONLINE_STRUCT_META(FAcceptFriendInvite::Params)
+	ONLINE_STRUCT_FIELD(FAcceptFriendInvite::Params, LocalAccountId),
+	ONLINE_STRUCT_FIELD(FAcceptFriendInvite::Params, TargetAccountId)
+END_ONLINE_STRUCT_META()
+
+BEGIN_ONLINE_STRUCT_META(FAcceptFriendInvite::Result)
+END_ONLINE_STRUCT_META()
+
+BEGIN_ONLINE_STRUCT_META(FRejectFriendInvite::Params)
+	ONLINE_STRUCT_FIELD(FRejectFriendInvite::Params, LocalAccountId),
+	ONLINE_STRUCT_FIELD(FRejectFriendInvite::Params, TargetAccountId)
+END_ONLINE_STRUCT_META()
+
+BEGIN_ONLINE_STRUCT_META(FRejectFriendInvite::Result)
+END_ONLINE_STRUCT_META()
+
+BEGIN_ONLINE_STRUCT_META(FRelationshipUpdated)
+	ONLINE_STRUCT_FIELD(FRelationshipUpdated, LocalAccountId),
+	ONLINE_STRUCT_FIELD(FRelationshipUpdated, RemoteAccountId),
+	ONLINE_STRUCT_FIELD(FRelationshipUpdated, OldRelationship),
+	ONLINE_STRUCT_FIELD(FRelationshipUpdated, NewRelationship)
+END_ONLINE_STRUCT_META()
+
+BEGIN_ONLINE_STRUCT_META(FQueryBlockedUsers::Params)
+	ONLINE_STRUCT_FIELD(FQueryBlockedUsers::Params, LocalAccountId)
+END_ONLINE_STRUCT_META()
+
+BEGIN_ONLINE_STRUCT_META(FQueryBlockedUsers::Result)
+END_ONLINE_STRUCT_META()
+
+BEGIN_ONLINE_STRUCT_META(FGetBlockedUsers::Params)
+	ONLINE_STRUCT_FIELD(FGetBlockedUsers::Params, LocalAccountId)
+END_ONLINE_STRUCT_META()
+
+BEGIN_ONLINE_STRUCT_META(FGetBlockedUsers::Result)
+	ONLINE_STRUCT_FIELD(FGetBlockedUsers::Result, BlockedUsers)
+END_ONLINE_STRUCT_META()
+
+BEGIN_ONLINE_STRUCT_META(FBlockUser::Params)
+	ONLINE_STRUCT_FIELD(FBlockUser::Params, LocalAccountId),
+	ONLINE_STRUCT_FIELD(FBlockUser::Params, TargetAccountId)
+END_ONLINE_STRUCT_META()
+
+BEGIN_ONLINE_STRUCT_META(FBlockUser::Result)
+END_ONLINE_STRUCT_META()
+
+/* Meta*/ }
+
+/* UE::Online */ }
diff --git a/Engine/Plugins/Online/Experimental/OnlineServicesEOS/Source/Private/Online/FriendsEOS.cpp b/Engine/Plugins/Online/Experimental/OnlineServicesEOS/Source/Private/Online/FriendsEOS.cpp
deleted file mode 100644
diff --git a/Engine/Plugins/Online/Experimental/OnlineServicesEOS/Source/Private/Online/OnlineServicesEOS.cpp b/Engine/Plugins/Online/Experimental/OnlineServicesEOS/Source/Private/Online/OnlineServicesEOS.cpp
--- a/Engine/Plugins/Online/Experimental/OnlineServicesEOS/Source/Private/Online/OnlineServicesEOS.cpp
+++ b/Engine/Plugins/Online/Experimental/OnlineServicesEOS/Source/Private/Online/OnlineServicesEOS.cpp
@@ -4,7 +4,7 @@
 
 #include "Online/AuthEOS.h"
 #include "Online/ExternalUIEOS.h"
-#include "Online/FriendsEOS.h"
+#include "Online/SocialEOS.h"
 #include "Online/PresenceEOS.h"
 #include "Online/UserInfoEOS.h"
 #include "Online/CommerceEOS.h"
@@ -20,7 +20,7 @@ void FOnlineServicesEOS::RegisterComponents()
 {
 	Components.Register<FAuthEOS>(*this);
 	Components.Register<FExternalUIEOS>(*this);
-	Components.Register<FFriendsEOS>(*this);
+	Components.Register<FSocialEOS>(*this);
 	Components.Register<FPresenceEOS>(*this);
 	Components.Register<FUserInfoEOS>(*this);
 	Components.Register<FCommerceEOS>(*this);
diff --git a/Engine/Plugins/Online/Experimental/OnlineServicesEOS/Source/Private/Online/SocialEOS.cpp b/Engine/Plugins/Online/Experimental/OnlineServicesEOS/Source/Private/Online/SocialEOS.cpp
new file mode 100644
--- /dev/null
+++ b/Engine/Plugins/Online/Experimental/OnlineServicesEOS/Source/Private/Online/SocialEOS.cpp
@@ -0,0 +1,425 @@
+// Copyright Epic Games, Inc. All Rights Reserved.
+
+#include "Online/SocialEOS.h"
+#include "Online/OnlineAsyncOp.h"
+#include "Online/OnlineServicesEOS.h"
+#include "Online/OnlineServicesEOSTypes.h"
+#include "Online/OnlineIdEOS.h"
+#include "Online/AuthEOS.h"
+
+#include "eos_friends.h"
+#include "eos_ui.h"
+
+namespace UE::Online {
+
+ERelationship EOSEFriendStatusToERelationship(EOS_EFriendsStatus Status)
+{
+	switch (Status)
+	{
+	case EOS_EFriendsStatus::EOS_FS_NotFriends:		return ERelationship::NotFriend;
+	case EOS_EFriendsStatus::EOS_FS_InviteSent:		return ERelationship::InviteSent;
+	case EOS_EFriendsStatus::EOS_FS_InviteReceived:	return ERelationship::InviteReceived;
+	case EOS_EFriendsStatus::EOS_FS_Friends:		return ERelationship::Friend;
+	}
+
+	return ERelationship::NotFriend;
+}
+
+FSocialEOS::FSocialEOS(FOnlineServicesEOS& InServices)
+	: FSocialCommon(InServices)
+{
+}
+
+void FSocialEOS::Initialize()
+{
+	FSocialCommon::Initialize();
+
+	FriendsHandle = EOS_Platform_GetFriendsInterface(static_cast<FOnlineServicesEOS&>(GetServices()).GetEOSPlatformHandle());
+	check(FriendsHandle != nullptr);
+
+	UIHandle = EOS_Platform_GetUIInterface(static_cast<FOnlineServicesEOS&>(GetServices()).GetEOSPlatformHandle());
+
+	// Register for friend updates
+	EOS_Friends_AddNotifyFriendsUpdateOptions Options = { };
+	Options.ApiVersion = EOS_FRIENDS_ADDNOTIFYFRIENDSUPDATE_API_LATEST;
+	NotifyFriendsUpdateNotificationId = EOS_Friends_AddNotifyFriendsUpdate(FriendsHandle, &Options, this, [](const EOS_Friends_OnFriendsUpdateInfo* Data)
+	{
+		FSocialEOS* This = reinterpret_cast<FSocialEOS*>(Data->ClientData);
+
+		const FAccountId LocalAccountId = FindAccountIdChecked(Data->LocalUserId);
+		This->Services.Get<FAuthEOS>()->ResolveAccountId(LocalAccountId, Data->TargetUserId)
+		.Next([This, LocalAccountId, PreviousStatus = Data->PreviousStatus, CurrentStatus = Data->CurrentStatus](const FAccountId& FriendAccountId)
+		{
+			This->OnEOSFriendsUpdate(LocalAccountId, FriendAccountId, PreviousStatus, CurrentStatus);
+		});
+	});
+}
+
+void FSocialEOS::PreShutdown()
+{
+	EOS_Friends_RemoveNotifyFriendsUpdate(FriendsHandle, NotifyFriendsUpdateNotificationId);
+}
+
+TOnlineAsyncOpHandle<FQueryFriends> FSocialEOS::QueryFriends(FQueryFriends::Params&& InParams)
+{
+	TOnlineAsyncOpRef<FQueryFriends> Op = GetJoinableOp<FQueryFriends>(MoveTemp(InParams));
+	if (Op->IsReady())
+	{
+		return Op->GetHandle();
+	}
+
+	Op->Then([this](TOnlineAsyncOp<FQueryFriends>& Op)
+	{
+		const FQueryFriends::Params& Params = Op.GetParams();
+
+		if (!Params.LocalAccountId.IsValid())
+		{
+			Op.SetError(Errors::InvalidParams());
+			return;
+		}
+
+		if (!Services.Get<FAuthEOS>()->IsLoggedIn(Params.LocalAccountId))
+		{
+			Op.SetError(Errors::NotLoggedIn());
+			return;
+		}
+	})
+	.Then([this](TOnlineAsyncOp<FQueryFriends>& Op, TPromise<const EOS_Friends_QueryFriendsCallbackInfo*>&& Promise)
+	{
+		EOS_Friends_QueryFriendsOptions QueryFriendsOptions = { };
+		QueryFriendsOptions.ApiVersion = EOS_FRIENDS_QUERYFRIENDS_API_LATEST;
+		QueryFriendsOptions.LocalUserId = GetEpicAccountIdChecked(Op.GetParams().LocalAccountId);
+
+		EOS_Async(EOS_Friends_QueryFriends, FriendsHandle, QueryFriendsOptions, MoveTemp(Promise));
+	})
+	.Then([this](TOnlineAsyncOp<FQueryFriends>& Op, const EOS_Friends_QueryFriendsCallbackInfo* Data)
+	{
+		if (Data->ResultCode != EOS_EResult::EOS_Success)
+		{
+			Op.SetError(Errors::FromEOSResult(Data->ResultCode));
+			return TArray<EOS_EpicAccountId>();
+		}
+
+		EOS_Friends_GetFriendsCountOptions GetFriendsCountOptions = {
+			EOS_FRIENDS_GETFRIENDSCOUNT_API_LATEST,
+			Data->LocalUserId
+		};
+
+		int32_t NumFriends = EOS_Friends_GetFriendsCount(FriendsHandle, &GetFriendsCountOptions);
+		TArray<EOS_EpicAccountId> CurrentFriends;
+		CurrentFriends.Reserve(NumFriends);
+		for (int32_t FriendIndex = 0; FriendIndex < NumFriends; ++FriendIndex)
+		{
+			EOS_Friends_GetFriendAtIndexOptions GetFriendAtIndexOptions = {
+				EOS_FRIENDS_GETFRIENDATINDEX_API_LATEST,
+				Data->LocalUserId,
+				FriendIndex
+			};
+
+			const EOS_EpicAccountId FriendEasId = EOS_Friends_GetFriendAtIndex(FriendsHandle, &GetFriendAtIndexOptions);
+			check(EOS_EpicAccountId_IsValid(FriendEasId));
+			CurrentFriends.Emplace(FriendEasId);
+		}
+		return CurrentFriends;
+	})
+	.Then(Services.Get<FAuthEOS>()->ResolveEpicIdsFn())
+	.Then([this](TOnlineAsyncOp<FQueryFriends>& Op, const TArray<FAccountId>& CurrentFriendIds)
+	{
+		const FAccountId LocalUserId = Op.GetParams().LocalAccountId;
+		const EOS_EpicAccountId LocalUserEasId = GetEpicAccountIdChecked(LocalUserId);
+
+		TMap<FAccountId, TSharedRef<FFriend>>& FriendsList = FriendsLists.FindOrAdd(LocalUserId);
+				
+		for (const FAccountId& FriendId : CurrentFriendIds)
+		{
+			const EOS_EpicAccountId FriendEasId = GetEpicAccountIdChecked(FriendId);
+
+			EOS_Friends_GetStatusOptions GetStatusOptions = {
+				EOS_FRIENDS_GETSTATUS_API_LATEST,
+				LocalUserEasId,
+				FriendEasId
+			};
+			const EOS_EFriendsStatus EOSFriendStatus = EOS_Friends_GetStatus(FriendsHandle, &GetStatusOptions);
+			const ERelationship Relationship = EOSEFriendStatusToERelationship(EOSFriendStatus);
+
+			if (TSharedRef<FFriend>* FriendPtr = FriendsList.Find(FriendId))
+			{
+				if (Relationship != (*FriendPtr)->Relationship)
+				{
+					ERelationship OldRelationship = (*FriendPtr)->Relationship;
+					(*FriendPtr)->Relationship = Relationship;
+
+					BroadcastRelationshipUpdated(LocalUserId, FriendId, OldRelationship, Relationship);
+				}
+			}
+			else
+			{
+				TSharedRef<FFriend> Friend = MakeShared<FFriend>();
+				Friend->FriendId = FriendId;
+				Friend->Relationship = Relationship;
+
+				FriendsList.Emplace(FriendId, Friend);
+
+				BroadcastRelationshipUpdated(LocalUserId, FriendId, ERelationship::NotFriend, Relationship);
+			}
+		}
+
+		Op.SetResult(FQueryFriends::Result());
+	})
+	.Enqueue(GetSerialQueue());
+
+	return Op->GetHandle();
+}
+
+TOnlineResult<FGetFriends> FSocialEOS::GetFriends(FGetFriends::Params&& Params)
+{
+	if (TMap<FAccountId, TSharedRef<FFriend>>* FriendsList = FriendsLists.Find(Params.LocalAccountId))
+	{
+		FGetFriends::Result Result;
+		FriendsList->GenerateValueArray(Result.Friends);
+		return TOnlineResult<FGetFriends>(MoveTemp(Result));
+	}
+	return TOnlineResult<FGetFriends>(Errors::InvalidState());
+}
+
+TOnlineAsyncOpHandle<FSendFriendInvite> FSocialEOS::SendFriendInvite(FSendFriendInvite::Params&& Params)
+{
+	TOnlineAsyncOpRef<FSendFriendInvite> Op = GetJoinableOp<FSendFriendInvite>(MoveTemp(Params));
+
+	if (Op->IsReady())
+	{
+		return Op->GetHandle();
+	}
+
+	Op->Then([this](TOnlineAsyncOp<FSendFriendInvite>& Op)
+	{
+		const FSendFriendInvite::Params& Params = Op.GetParams();
+
+		if (!Params.LocalAccountId.IsValid() || !Params.TargetAccountId.IsValid() || !EOS_EpicAccountId_IsValid(GetEpicAccountIdChecked(Params.TargetAccountId)))
+		{
+			Op.SetError(Errors::InvalidParams());
+			return;
+		}
+
+		if (!Services.Get<FAuthEOS>()->IsLoggedIn(Params.LocalAccountId))
+		{
+			Op.SetError(Errors::NotLoggedIn());
+			return;
+		}
+	})
+	.Then([this](TOnlineAsyncOp<FSendFriendInvite>& Op, TPromise<const EOS_Friends_SendInviteCallbackInfo*>&& Promise)
+	{
+		const FSendFriendInvite::Params& Params = Op.GetParams();
+
+		EOS_Friends_SendInviteOptions SendInviteOptions = {};
+		SendInviteOptions.ApiVersion = EOS_FRIENDS_SENDINVITE_API_LATEST;
+		SendInviteOptions.LocalUserId = GetEpicAccountIdChecked(Params.LocalAccountId);
+		SendInviteOptions.TargetUserId = GetEpicAccountIdChecked(Params.TargetAccountId);
+
+		EOS_Async(EOS_Friends_SendInvite, FriendsHandle, SendInviteOptions, MoveTemp(Promise));
+	})
+	.Then([this](TOnlineAsyncOp<FSendFriendInvite>& Op, const EOS_Friends_SendInviteCallbackInfo* Data)
+	{
+		if (Data->ResultCode != EOS_EResult::EOS_Success)
+		{
+			Op.SetError(Errors::FromEOSResult(Data->ResultCode));
+			return;
+		}
+
+		Op.SetResult({});
+	})
+	.Enqueue(GetSerialQueue());
+
+	return Op->GetHandle();
+}
+
+TOnlineAsyncOpHandle<FAcceptFriendInvite> FSocialEOS::AcceptFriendInvite(FAcceptFriendInvite::Params&& Params)
+{
+	TOnlineAsyncOpRef<FAcceptFriendInvite> Op = GetJoinableOp<FAcceptFriendInvite>(MoveTemp(Params));
+
+	if (Op->IsReady())
+	{
+		return Op->GetHandle();
+	}
+
+	Op->Then([this](TOnlineAsyncOp<FAcceptFriendInvite>& Op)
+	{
+		const FAcceptFriendInvite::Params& Params = Op.GetParams();
+
+		if (!Params.LocalAccountId.IsValid() || !Params.TargetAccountId.IsValid() || !EOS_EpicAccountId_IsValid(GetEpicAccountIdChecked(Params.TargetAccountId)))
+		{
+			Op.SetError(Errors::InvalidParams());
+			return;
+		}
+
+		if (!Services.Get<FAuthEOS>()->IsLoggedIn(Params.LocalAccountId))
+		{
+			Op.SetError(Errors::NotLoggedIn());
+			return;
+		}
+	})
+	.Then([this](TOnlineAsyncOp<FAcceptFriendInvite>& Op, TPromise<const EOS_Friends_AcceptInviteCallbackInfo*>&& Promise)
+	{
+		const FAcceptFriendInvite::Params& Params = Op.GetParams();
+
+		EOS_Friends_AcceptInviteOptions AcceptInviteOptions = {};
+		AcceptInviteOptions.ApiVersion = EOS_FRIENDS_ACCEPTINVITE_API_LATEST;
+		AcceptInviteOptions.LocalUserId = GetEpicAccountIdChecked(Params.LocalAccountId);
+		AcceptInviteOptions.TargetUserId = GetEpicAccountIdChecked(Params.TargetAccountId);
+
+		EOS_Async(EOS_Friends_AcceptInvite, FriendsHandle, AcceptInviteOptions, MoveTemp(Promise));
+	})
+	.Then([this](TOnlineAsyncOp<FAcceptFriendInvite>& Op, const EOS_Friends_AcceptInviteCallbackInfo* Data)
+	{
+		if (Data->ResultCode != EOS_EResult::EOS_Success)
+		{
+			Op.SetError(Errors::FromEOSResult(Data->ResultCode));
+			return;
+		}
+
+		Op.SetResult({});
+	})
+	.Enqueue(GetSerialQueue());
+
+	return Op->GetHandle();
+}
+
+TOnlineAsyncOpHandle<FRejectFriendInvite> FSocialEOS::RejectFriendInvite(FRejectFriendInvite::Params&& Params)
+{
+	TOnlineAsyncOpRef<FRejectFriendInvite> Op = GetJoinableOp<FRejectFriendInvite>(MoveTemp(Params));
+
+	if (Op->IsReady())
+	{
+		return Op->GetHandle();
+	}
+
+	Op->Then([this](TOnlineAsyncOp<FRejectFriendInvite>& Op)
+	{
+		const FRejectFriendInvite::Params& Params = Op.GetParams();
+
+		if (!Params.LocalAccountId.IsValid() || !Params.TargetAccountId.IsValid() || !EOS_EpicAccountId_IsValid(GetEpicAccountIdChecked(Params.TargetAccountId)))
+		{
+			Op.SetError(Errors::InvalidParams());
+			return;
+		}
+
+		if (!Services.Get<FAuthEOS>()->IsLoggedIn(Params.LocalAccountId))
+		{
+			Op.SetError(Errors::NotLoggedIn());
+			return;
+		}
+	})
+	.Then([this](TOnlineAsyncOp<FRejectFriendInvite>& Op, TPromise<const EOS_Friends_RejectInviteCallbackInfo*>&& Promise)
+	{
+		const FRejectFriendInvite::Params& Params = Op.GetParams();
+
+		EOS_Friends_RejectInviteOptions RejectInviteOptions = {};
+		RejectInviteOptions.ApiVersion = EOS_FRIENDS_REJECTINVITE_API_LATEST;
+		RejectInviteOptions.LocalUserId = GetEpicAccountIdChecked(Params.LocalAccountId);
+		RejectInviteOptions.TargetUserId = GetEpicAccountIdChecked(Params.TargetAccountId);
+
+		EOS_Async(EOS_Friends_RejectInvite, FriendsHandle, RejectInviteOptions, MoveTemp(Promise));
+	})
+	.Then([this](TOnlineAsyncOp<FRejectFriendInvite>& Op, const EOS_Friends_RejectInviteCallbackInfo* Data)
+	{
+		if (Data->ResultCode != EOS_EResult::EOS_Success)
+		{
+			Op.SetError(Errors::FromEOSResult(Data->ResultCode));
+			return;
+		}
+
+		Op.SetResult({});
+	})
+	.Enqueue(GetSerialQueue());
+
+	return Op->GetHandle();
+}
+
+TOnlineAsyncOpHandle<FBlockUser> FSocialEOS::BlockUser(FBlockUser::Params&& Params)
+{
+	TOnlineAsyncOpRef<FBlockUser> Op = GetJoinableOp<FBlockUser>(MoveTemp(Params));
+
+	if (Op->IsReady())
+	{
+		return Op->GetHandle();
+	}
+
+	Op->Then([this](TOnlineAsyncOp<FBlockUser>& Op)
+	{
+		const FBlockUser::Params& Params = Op.GetParams();
+
+		if (!UIHandle)
+		{
+			Op.SetError(Errors::MissingInterface());
+			return;
+		}
+
+		if (!Params.LocalAccountId.IsValid() || !Params.TargetAccountId.IsValid() || !EOS_EpicAccountId_IsValid(GetEpicAccountIdChecked(Params.TargetAccountId)))
+		{
+			Op.SetError(Errors::InvalidParams());
+			return;
+		}
+
+		if (!Services.Get<FAuthEOS>()->IsLoggedIn(Params.LocalAccountId))
+		{
+			Op.SetError(Errors::NotLoggedIn());
+			return;
+		}
+	})
+	.Then([this](TOnlineAsyncOp<FBlockUser>& Op, TPromise<const EOS_UI_OnShowBlockPlayerCallbackInfo*>&& Promise)
+	{
+		const FBlockUser::Params& Params = Op.GetParams();
+
+		EOS_UI_ShowBlockPlayerOptions ShowBlockPlayerOptions = {};
+		ShowBlockPlayerOptions.ApiVersion = EOS_UI_SHOWBLOCKPLAYER_API_LATEST;
+		ShowBlockPlayerOptions.LocalUserId = GetEpicAccountIdChecked(Params.LocalAccountId);
+		ShowBlockPlayerOptions.TargetUserId = GetEpicAccountIdChecked(Params.TargetAccountId);
+
+		EOS_Async(EOS_UI_ShowBlockPlayer, UIHandle, ShowBlockPlayerOptions, MoveTemp(Promise));
+	})
+	.Then([this](TOnlineAsyncOp<FBlockUser>& Op, const EOS_UI_OnShowBlockPlayerCallbackInfo* Data)
+	{
+		if (Data->ResultCode != EOS_EResult::EOS_Success)
+		{
+			Op.SetError(Errors::FromEOSResult(Data->ResultCode));
+			return;
+		}
+
+		Op.SetResult({});
+	})
+	.Enqueue(GetSerialQueue());
+
+	return Op->GetHandle();
+}
+
+void FSocialEOS::OnEOSFriendsUpdate(FAccountId LocalAccountId, FAccountId FriendAccountId, EOS_EFriendsStatus PreviousStatus, EOS_EFriendsStatus CurrentStatus)
+{
+	const EOS_EpicAccountId LocalUserEasId = GetEpicAccountIdChecked(LocalAccountId);
+
+	TMap<FAccountId, TSharedRef<FFriend>>& FriendsList = FriendsLists.FindOrAdd(LocalAccountId);
+
+	if (CurrentStatus == EOS_EFriendsStatus::EOS_FS_NotFriends)
+	{
+		FriendsList.Remove(FriendAccountId);
+	}
+	else
+	{
+		// Add/update friend with appropriate status
+		if (TSharedRef<FFriend>* ExistingFriend = FriendsList.Find(FriendAccountId))
+		{
+			(*ExistingFriend)->Relationship = EOSEFriendStatusToERelationship(CurrentStatus);
+		}
+		else
+		{
+			// Create new entry
+			TSharedRef<FFriend> Friend = MakeShared<FFriend>();
+			Friend->FriendId = FriendAccountId;
+			Friend->Relationship = EOSEFriendStatusToERelationship(CurrentStatus);
+			FriendsList.Emplace(FriendAccountId, MoveTemp(Friend));
+		}
+	}
+
+	BroadcastRelationshipUpdated(LocalAccountId, FriendAccountId, EOSEFriendStatusToERelationship(PreviousStatus), EOSEFriendStatusToERelationship(CurrentStatus));
+}
+
+/* UE::Online */ }
diff --git a/Engine/Plugins/Online/Experimental/OnlineServicesEOS/Source/Public/Online/FriendsEOS.h b/Engine/Plugins/Online/Experimental/OnlineServicesEOS/Source/Public/Online/FriendsEOS.h
deleted file mode 100644
diff --git a/Engine/Plugins/Online/Experimental/OnlineServicesEOS/Source/Public/Online/SocialEOS.h b/Engine/Plugins/Online/Experimental/OnlineServicesEOS/Source/Public/Online/SocialEOS.h
new file mode 100644
--- /dev/null
+++ b/Engine/Plugins/Online/Experimental/OnlineServicesEOS/Source/Public/Online/SocialEOS.h
@@ -0,0 +1,43 @@
+// Copyright Epic Games, Inc. All Rights Reserved.
+
+#pragma once
+
+#include "CoreMinimal.h"
+#include "Online/SocialCommon.h"
+
+#if defined(EOS_PLATFORM_BASE_FILE_NAME)
+#include EOS_PLATFORM_BASE_FILE_NAME
+#endif
+#include "eos_friends_types.h"
+#include "eos_ui_types.h"
+
+namespace UE::Online {
+
+class FOnlineServicesEOS;
+	
+class FSocialEOS : public FSocialCommon
+{
+public:
+	FSocialEOS(FOnlineServicesEOS& InServices);
+
+	virtual void Initialize() override;
+	virtual void PreShutdown() override;
+
+	virtual TOnlineAsyncOpHandle<FQueryFriends> QueryFriends(FQueryFriends::Params&& Params) override;
+	virtual TOnlineResult<FGetFriends> GetFriends(FGetFriends::Params&& Params) override;
+	virtual TOnlineAsyncOpHandle<FSendFriendInvite> SendFriendInvite(FSendFriendInvite::Params&& Params) override;
+	virtual TOnlineAsyncOpHandle<FAcceptFriendInvite> AcceptFriendInvite(FAcceptFriendInvite::Params&& Params) override;
+	virtual TOnlineAsyncOpHandle<FRejectFriendInvite> RejectFriendInvite(FRejectFriendInvite::Params&& Params) override;
+	virtual TOnlineAsyncOpHandle<FBlockUser> BlockUser(FBlockUser::Params&& Params) override;
+
+protected:
+	void OnEOSFriendsUpdate(FAccountId LocalAccountId, FAccountId FriendAccountId, EOS_EFriendsStatus PreviousStatus, EOS_EFriendsStatus CurrentStatus);
+
+	TMap<FAccountId, TMap<FAccountId, TSharedRef<FFriend>>> FriendsLists;
+	EOS_HFriends FriendsHandle = nullptr;
+	EOS_HUI UIHandle = nullptr;
+
+	EOS_NotificationId NotifyFriendsUpdateNotificationId = 0;
+};
+
+/* UE::Online */ }
