Fix for Horde Issue 219110

diff --git a/Engine/Plugins/Online/Experimental/OnlineServicesEOS/Source/Private/OnlineServicesEOSModule.cpp b/Engine/Plugins/Online/Experimental/OnlineServicesEOS/Source/Private/OnlineServicesEOSModule.cpp
--- a/Engine/Plugins/Online/Experimental/OnlineServicesEOS/Source/Private/OnlineServicesEOSModule.cpp
+++ b/Engine/Plugins/Online/Experimental/OnlineServicesEOS/Source/Private/OnlineServicesEOSModule.cpp
@@ -38,6 +38,12 @@ void FOnlineServicesEOSModule::StartupModule()
 		FModuleManager::Get().LoadModuleChecked(EOSSharedModuleName);
 	}
 
+	const FName OnlineServicesInterfaceModuleName = TEXT("OnlineServicesInterface");
+	if (!FModuleManager::Get().IsModuleLoaded(OnlineServicesInterfaceModuleName))
+	{
+		FModuleManager::Get().LoadModuleChecked(OnlineServicesInterfaceModuleName);
+	}
+
 	// Make this higher priority that EOSGS
 	const int Priority = 1;
 	FOnlineServicesRegistry::Get().RegisterServicesFactory(EOnlineServices::Epic, MakeUnique<FOnlineServicesFactoryEOS>(), Priority);
