Move OnlineServices plugins out of Experimental

diff --git a/Engine/Plugins/Online/OnlineServices/Source/OnlineServicesCommon/Public/Online/ExternalUICommon.h b/Engine/Plugins/Online/OnlineServices/Source/OnlineServicesCommon/Public/Online/ExternalUICommon.h
new file mode 100644
--- /dev/null
+++ b/Engine/Plugins/Online/OnlineServices/Source/OnlineServicesCommon/Public/Online/ExternalUICommon.h
@@ -0,0 +1,31 @@
+// Copyright Epic Games, Inc. All Rights Reserved.
+
+#pragma once
+
+#include "Online/ExternalUI.h"
+#include "Online/OnlineComponent.h"
+
+namespace UE::Online {
+
+class FOnlineServicesCommon;
+
+class ONLINESERVICESCOMMON_API FExternalUICommon : public TOnlineComponent<IExternalUI>
+{
+public:
+	using Super = IExternalUI;
+
+	FExternalUICommon(FOnlineServicesCommon& InServices);
+
+	// TOnlineComponent
+	virtual void RegisterCommands() override;
+
+	// IExternalUI
+	virtual TOnlineAsyncOpHandle<FExternalUIShowLoginUI> ShowLoginUI(FExternalUIShowLoginUI::Params&& Params) override;
+	virtual TOnlineAsyncOpHandle<FExternalUIShowFriendsUI> ShowFriendsUI(FExternalUIShowFriendsUI::Params&& Params) override;
+
+	virtual TOnlineEvent<void(const FExternalUIStatusChanged&)> OnExternalUIStatusChanged() override;
+protected:
+	TOnlineEventCallable<void(const FExternalUIStatusChanged&)> OnExternalUIStatusChangedEvent;
+};
+
+/* UE::Online */ }
\ No newline at end of file
diff --git a/Engine/Plugins/Online/OnlineServices/Source/OnlineServicesCommon/Public/Online/OnlineServicesCommon.h b/Engine/Plugins/Online/OnlineServices/Source/OnlineServicesCommon/Public/Online/OnlineServicesCommon.h
new file mode 100644
--- /dev/null
+++ b/Engine/Plugins/Online/OnlineServices/Source/OnlineServicesCommon/Public/Online/OnlineServicesCommon.h
@@ -0,0 +1,342 @@
+// Copyright Epic Games, Inc. All Rights Reserved.
+
+#pragma once
+
+#include "Online/OnlineServices.h"
+
+#include "Online/OnlineComponentRegistry.h"
+#include "Online/OnlineAsyncOpCache.h"
+#include "Online/OnlineAsyncOpQueue.h"
+#include "Online/OnlineConfig.h"
+#include "Online/OnlineExecHandler.h"
+#include "Online/OnlineServicesLog.h"
+
+#include "Async/Async.h"
+#include "Containers/Ticker.h"
+#include "Misc/CoreMisc.h"
+#include "Templates/SharedPointer.h"
+
+namespace UE::Online {
+
+class ONLINESERVICESCOMMON_API FOnlineServicesCommon
+	: public IOnlineServices
+	, public TSharedFromThis<FOnlineServicesCommon>
+	, public FTSTickerObjectBase
+	, public FSelfRegisteringExec
+{
+public:
+	using Super = IOnlineServices;
+
+	FOnlineServicesCommon(const FString& InConfigName, FName InInstanceName);
+	FOnlineServicesCommon(const FOnlineServicesCommon&) = delete;
+	FOnlineServicesCommon(FOnlineServicesCommon&&) = delete;
+	virtual ~FOnlineServicesCommon() {}
+
+	// IOnlineServices
+	virtual void Init() override;
+	virtual void Destroy() override;
+	virtual IAchievementsPtr GetAchievementsInterface() override;
+	virtual IAuthPtr GetAuthInterface() override;
+	virtual IUserInfoPtr GetUserInfoInterface() override;
+	virtual ICommercePtr GetCommerceInterface() override;
+	virtual ISocialPtr GetSocialInterface() override;
+	virtual IPresencePtr GetPresenceInterface() override;
+	virtual IExternalUIPtr GetExternalUIInterface() override;
+	virtual ILeaderboardsPtr GetLeaderboardsInterface() override;
+	virtual ILobbiesPtr GetLobbiesInterface() override;
+	virtual ISessionsPtr GetSessionsInterface() override;
+	virtual IStatsPtr GetStatsInterface() override;
+	virtual IConnectivityPtr GetConnectivityInterface() override;
+	virtual IPrivilegesPtr GetPrivilegesInterface() override;
+	virtual ITitleFilePtr GetTitleFileInterface() override;
+	virtual IUserFilePtr GetUserFileInterface() override;
+	virtual TOnlineResult<FGetResolvedConnectString> GetResolvedConnectString(FGetResolvedConnectString::Params&& Params) override;
+	virtual FName GetInstanceName() const override;
+
+	// FOnlineServicesCommon
+
+	/**
+	 * Retrieve any of the Interface IOnlineComponents
+	 */
+	template <typename ComponentType>
+	ComponentType* Get()
+	{
+		return Components.Get<ComponentType>();
+	}
+
+	/**
+	 * Called to register all the IOnlineComponents with the IOnlineService, called after this is constructed
+	 */
+	virtual void RegisterComponents();
+
+	/**
+	 * Calls Initialize on all the components, called after RegisterComponents
+	 */
+	virtual void Initialize();
+
+	/**
+	 * Calls PostInitialize on all the components, called after Initialize
+	 */
+	virtual void PostInitialize();
+
+	/**
+	 * Calls UpdateConfig on all the components
+	 */
+	virtual void UpdateConfig();
+
+	/**
+	 * Calls Tick on all the components
+	 */
+	virtual bool Tick(float DeltaSeconds) override;
+
+	/**
+	 * Calls PreShutdown on all the components, called prior to Shutdown
+	 */
+	virtual void PreShutdown();
+
+	/**
+	 * Calls Shutdown on all the components, called before this is destructed
+	 */
+	virtual void Shutdown();
+
+	/**
+	 * Call a callable according to a specified execution policy
+	 */
+	template <typename CallableType>
+	void Execute(FOnlineAsyncExecutionPolicy ExecutionPolicy, CallableType&& Callable)
+	{
+		switch (ExecutionPolicy.GetExecutionPolicy())
+		{
+			case EOnlineAsyncExecutionPolicy::RunOnGameThread:
+				ExecuteOnGameThread(MoveTemp(Callable));
+				break;
+
+			case EOnlineAsyncExecutionPolicy::RunOnNextTick:
+				Async(EAsyncExecution::TaskGraphMainThread, MoveTemp(Callable));
+				break;
+
+			case EOnlineAsyncExecutionPolicy::RunOnThreadPool:
+				Async(EAsyncExecution::ThreadPool, MoveTemp(Callable));
+				break;
+
+			case EOnlineAsyncExecutionPolicy::RunOnTaskGraph:
+				Async(EAsyncExecution::TaskGraph, MoveTemp(Callable));
+				break;
+
+			case EOnlineAsyncExecutionPolicy::RunImmediately:
+				Callable();
+				break;
+		}
+	}
+
+	/**
+	 * Call a callable on the game thread
+	 */
+	template <typename CallableType>
+	void ExecuteOnGameThread(CallableType&& Callable)
+	{
+		if (IsInGameThread())
+		{
+			Callable();
+		}
+		else
+		{
+			Async(EAsyncExecution::TaskGraphMainThread, MoveTemp(Callable));
+		}
+	}
+
+	/**
+	 * Override the default config provider (FOnlineConfigProviderGConfig(GEngineini))
+	 */
+	void SetConfigProvider(TUniquePtr<IOnlineConfigProvider>&& InConfigProvider)
+	{
+		ConfigProvider = MoveTemp(InConfigProvider);
+	}
+
+	/**
+	 * Clear the list of config overrides
+	 */
+	void ResetConfigSectionOverrides()
+	{
+		ConfigSectionOverrides.Reset();
+	}
+
+	/**
+	 * Add a config section override. These will be used in the order they are added
+	 */
+	void AddConfigSectionOverride(const FString& Override)
+	{
+		ConfigSectionOverrides.Add(Override);
+	}
+
+	/**
+	 * Get the config name for the Subsystem
+	 */
+	const FString& GetConfigName() const { return ConfigName; }
+
+	TArray<FString> GetConfigSectionHeiarchy(const FString& OperationName = FString()) const
+	{
+		TArray<FString> SectionHeiarchy;
+		FString SectionName = TEXT("OnlineServices");
+		SectionHeiarchy.Add(SectionName);
+		SectionName += TEXT(".") + GetConfigName();
+		SectionHeiarchy.Add(SectionName);
+		if (!OperationName.IsEmpty())
+		{
+			SectionName += TEXT(".") + OperationName;
+			SectionHeiarchy.Add(SectionName);
+		}
+		return SectionHeiarchy;
+	}
+
+	/**
+	 * Load a config struct for an interface + operation
+	 * Will load values from the following sections:
+	 *   OnlineServices
+	 *   OnlineServices.<InterfaceName> (if InterfaceName is set)
+	 *   OnlineServices.<ServiceProvider>
+	 *   OnlineServices.<ServiceProvider>.<InterfaceName> (if InterfaceName is set)
+	 *   OnlineServices.<ServiceProvider>.<InterfaceName>.<OperationName> (if OperationName is set)
+	 * 
+	 * @param Struct Struct to populate with values from config
+	 * @param InterfaceName Optional interface name to append to the config section name
+	 * @param OperationName Optional operation name to append to the config section name
+	 * 
+	 * @return true if a value was loaded
+	 */
+	template <typename StructType>
+	bool LoadConfig(StructType& Struct, const FString& InterfaceName = FString(), const FString& OperationName = FString()) const
+	{
+		TArray<FString> SectionHeiarchy;
+		FString SectionName = TEXT("OnlineServices");
+		SectionHeiarchy.Add(SectionName);
+		if (!InterfaceName.IsEmpty())
+		{
+			SectionHeiarchy.Add(SectionName + TEXT(".") + InterfaceName);
+		}
+		SectionName += TEXT(".") + GetConfigName();
+		SectionHeiarchy.Add(SectionName);
+		if (!InterfaceName.IsEmpty())
+		{
+			SectionName += TEXT(".") + InterfaceName;
+			SectionHeiarchy.Add(SectionName);
+			if (!OperationName.IsEmpty())
+			{
+				SectionName += TEXT(".") + OperationName;
+				SectionHeiarchy.Add(SectionName);
+			}
+		}
+		return LoadConfig(Struct, SectionHeiarchy);
+	}
+
+	/**
+	 * Get an array of a config section with the overrides added in
+	 * 
+	 * @param SectionHeiarchy Array of config sections to load values from
+	 * 
+	 * @return Array of the sections with overrides for values to be loaded from
+	 */
+	TArray<FString> GetConfigSectionHeirachWithOverrides(const TArray<FString>& SectionHeiarchy) const
+	{
+		TArray<FString> SectionHeiarchyWithOverrides;
+		for (const FString& Section : SectionHeiarchy)
+		{
+			SectionHeiarchyWithOverrides.Add(Section);
+			for (const FString& Override : ConfigSectionOverrides)
+			{
+				FString OverrideSection = Section + TEXT(" ") + Override;
+				SectionHeiarchyWithOverrides.Add(OverrideSection);
+			}
+		}
+
+		return SectionHeiarchyWithOverrides;
+	}
+
+	/**
+	 * Load a config struct for a section heiarchy, also using the ConfigSectionOverrides
+	 *
+	 * @param Struct Struct to populate with values from config
+	 * @param SectionHeiarchy Array of config sections to load values from
+	 *
+	 * @return true if a value was loaded
+	 */
+	template <typename StructType>
+	bool LoadConfig(StructType& Struct, const TArray<FString>& SectionHeiarchy) const
+	{
+		return ::UE::Online::LoadConfig(*ConfigProvider, GetConfigSectionHeirachWithOverrides(SectionHeiarchy), Struct);
+	}
+
+	/* Get op (OnlineServices) */
+	template <typename OpType>
+	TOnlineAsyncOpRef<OpType> GetOp(typename OpType::Params&& Params)
+	{
+		return OpCache.GetOp<OpType>(MoveTemp(Params), GetConfigSectionHeiarchy());
+	}
+
+	template <typename OpType, typename ParamsFuncsType = TJoinableOpParamsFuncs<OpType>>
+	TOnlineAsyncOpRef<OpType> GetJoinableOp(typename OpType::Params&& Params)
+	{
+		return OpCache.GetJoinableOp<OpType, ParamsFuncsType>(MoveTemp(Params), GetConfigSectionHeiarchy());
+	}
+
+	template <typename OpType, typename ParamsFuncsType = TMergeableOpParamsFuncs<OpType>>
+	TOnlineAsyncOpRef<OpType> GetMergeableOp(typename OpType::Params&& Params)
+	{
+		return OpCache.GetMergeableOp<OpType, ParamsFuncsType>(MoveTemp(Params), GetConfigSectionHeiarchy());
+	}
+
+	/* Get op (Interface) */
+	template <typename OpType>
+	TOnlineAsyncOpRef<OpType> GetOp(typename OpType::Params&& Params, const TArray<FString> ConfigSectionHeiarchy)
+	{
+		return OpCache.GetOp<OpType>(MoveTemp(Params), ConfigSectionHeiarchy);
+	}
+
+	template <typename OpType, typename ParamsFuncsType /*= TJoinableOpParamsFuncs<OpType>*/>
+	TOnlineAsyncOpRef<OpType> GetJoinableOp(typename OpType::Params&& Params, const TArray<FString> ConfigSectionHeiarchy)
+	{
+		return OpCache.GetJoinableOp<OpType, ParamsFuncsType>(MoveTemp(Params), ConfigSectionHeiarchy);
+	}
+
+	template <typename OpType, typename ParamsFuncsType /*= TMergeableOpParamsFuncs<OpType>*/>
+	TOnlineAsyncOpRef<OpType> GetMergeableOp(typename OpType::Params&& Params, const TArray<FString> ConfigSectionHeiarchy)
+	{
+		return OpCache.GetMergeableOp<OpType, ParamsFuncsType>(MoveTemp(Params), ConfigSectionHeiarchy);
+	}
+
+	/* Queue for executing tasks in parallel. Serial queues feed into this */
+	FOnlineAsyncOpQueueParallel& GetParallelQueue();
+
+	/* Queue for executing tasks in serial */
+	FOnlineAsyncOpQueue& GetSerialQueue();
+
+	/* Queues for executing per-user tasks in serial */
+	FOnlineAsyncOpQueue& GetSerialQueue(const FAccountId& AccountId);
+	
+	void RegisterExecHandler(const FString& Name, TUniquePtr<IOnlineExecHandler>&& Handler);
+
+	virtual bool Exec(UWorld* World, const TCHAR* Cmd, FOutputDevice& Ar) override;
+
+	FOnlineAsyncOpCache OpCache;
+
+protected:
+	TMap<FString, TUniquePtr<IOnlineExecHandler>> ExecCommands;
+
+	static uint32 NextInstanceIndex;
+	uint32 InstanceIndex;
+	FName InstanceName;
+
+	FOnlineComponentRegistry Components;
+	TUniquePtr<IOnlineConfigProvider> ConfigProvider;
+
+	/* Config section overrides */
+	TArray<FString> ConfigSectionOverrides;
+	FString ConfigName;
+
+	FOnlineAsyncOpQueueParallel ParallelQueue;
+	FOnlineAsyncOpQueueSerial SerialQueue;
+	TMap<FAccountId, TUniquePtr<FOnlineAsyncOpQueueSerial>> PerUserSerialQueue;
+};
+
+/* UE::Online */ }
+
diff --git a/Engine/Plugins/Online/OnlineServices/Source/OnlineServicesCommon/Public/Online/OnlineServicesLog.h b/Engine/Plugins/Online/OnlineServices/Source/OnlineServicesCommon/Public/Online/OnlineServicesLog.h
new file mode 100644
--- /dev/null
+++ b/Engine/Plugins/Online/OnlineServices/Source/OnlineServicesCommon/Public/Online/OnlineServicesLog.h
@@ -0,0 +1,7 @@
+// Copyright Epic Games, Inc. All Rights Reserved.
+
+#pragma once
+
+#include "Logging/LogMacros.h"
+
+ONLINESERVICESCOMMON_API DECLARE_LOG_CATEGORY_EXTERN(LogOnlineServices, Log, All);
diff --git a/Engine/Plugins/Online/OnlineServices/Source/OnlineServicesCommon/Public/Online/PresenceCommon.h b/Engine/Plugins/Online/OnlineServices/Source/OnlineServicesCommon/Public/Online/PresenceCommon.h
new file mode 100644
--- /dev/null
+++ b/Engine/Plugins/Online/OnlineServices/Source/OnlineServicesCommon/Public/Online/PresenceCommon.h
@@ -0,0 +1,33 @@
+// Copyright Epic Games, Inc. All Rights Reserved.
+
+#pragma once
+
+#include "Online/Presence.h"
+#include "Online/OnlineComponent.h"
+
+namespace UE::Online {
+
+class FOnlineServicesCommon;
+
+class ONLINESERVICESCOMMON_API FPresenceCommon : public TOnlineComponent<IPresence>
+{
+public:
+	using Super = IPresence;
+
+	FPresenceCommon(FOnlineServicesCommon& InServices);
+
+	// IPresence
+	virtual TOnlineAsyncOpHandle<FQueryPresence> QueryPresence(FQueryPresence::Params&& Params) override;
+	virtual TOnlineAsyncOpHandle<FBatchQueryPresence> BatchQueryPresence(FBatchQueryPresence::Params&& Params) override;
+	virtual TOnlineResult<FGetCachedPresence> GetCachedPresence(FGetCachedPresence::Params&& Params) override;
+	virtual TOnlineAsyncOpHandle<FUpdatePresence> UpdatePresence(FUpdatePresence::Params&& Params) override;
+	virtual TOnlineAsyncOpHandle<FPartialUpdatePresence> PartialUpdatePresence(FPartialUpdatePresence::Params&& Params) override;
+	virtual TOnlineEvent<void(const FPresenceUpdated&)> OnPresenceUpdated() override;
+
+protected:
+	FOnlineServicesCommon& Services;
+
+	TOnlineEventCallable<void(const FPresenceUpdated&)> OnPresenceUpdatedEvent;
+};
+
+/* UE::Online */ }
diff --git a/Engine/Plugins/Online/OnlineServices/Source/OnlineServicesInterface/Public/Online/OnlineErrorCode.h b/Engine/Plugins/Online/OnlineServices/Source/OnlineServicesInterface/Public/Online/OnlineErrorCode.h
new file mode 100644
--- /dev/null
+++ b/Engine/Plugins/Online/OnlineServices/Source/OnlineServicesInterface/Public/Online/OnlineErrorCode.h
@@ -0,0 +1,91 @@
+// Copyright Epic Games, Inc. All Rights Reserved.
+#pragma once
+namespace UE::Online::Errors {
+using ErrorCodeType = uint64;
+
+
+// error code as struct with bitfields
+struct FOnlineErrorCode
+{
+	uint64 System : 4; // Engine, Game, Third Party Plugin
+	uint64 Category : 28; // HRESULT, Http, nn::Result, etc
+	uint64 Code : 32; // HRESULT value, Http status code, nn::Result underlying int value, etc
+};
+
+namespace ErrorCode
+{
+	static constexpr ErrorCodeType Success = ErrorCodeType(0);
+
+	ONLINESERVICESINTERFACE_API FString ToString(ErrorCodeType ErrorCode);
+
+	constexpr ErrorCodeType Create(uint32 Source, uint32 Category, uint32 Code)
+	{
+		return ((Source & 0x3ull) << 60ull) | ((Category & 0x3fffffffull) << 32ull) | static_cast<ErrorCodeType>(Code);
+	}
+
+	constexpr ErrorCodeType Create(uint32 Category, uint32 Code)
+	{
+		return ((Category & 0x3fffffffull) << 32ull) | static_cast<ErrorCodeType>(Code);
+	}
+
+	namespace System
+	{
+		const ErrorCodeType Engine = 1;
+		const ErrorCodeType Game = 2;
+		const ErrorCodeType ThirdPartyPlugin = 3;
+	}
+}
+
+// These three functions are responsible for reading the parts of the uint64 error representation
+// The uint64 error is formatted as follows (big endian):
+// ssss cccc cccc cccc cccc cccc cccc cccc cccc vvvv vvvv vvvv vvvv vvvv vvvv vvvv vvvv 
+ONLINESERVICESINTERFACE_API uint64 ErrorCodeSystem(ErrorCodeType ErrorCode);
+ONLINESERVICESINTERFACE_API uint64 ErrorCodeCategory(ErrorCodeType ErrorCode);
+ONLINESERVICESINTERFACE_API uint64 ErrorCodeValue(ErrorCodeType ErrorCode); 
+
+// needs to be inside Online::V2::Errors namespace
+// define an error category
+#define UE_ONLINE_ERROR_CATEGORY(Name, InSystem, Value, Description) \
+namespace ErrorCode { namespace Category { static constexpr int Name = Value; } }\
+namespace ErrorCode { namespace Category { static constexpr int Name##_System = UE::Online::Errors::ErrorCode::System::InSystem; } }  
+
+// reserve an error category range
+#define UE_ONLINE_ERROR_CATEGORY_RANGE(Name, System, StartValue, EndValue, Description) \
+namespace ErrorCode { namespace Category { \
+    static constexpr Name##_Start = StartValue; \
+    static constexpr Name##_End = EndValue; \
+} }
+
+// Defines an individual error 
+//	- Example access: (UE::Online::Errors::)NoConnection();
+	// - Can also specify an inner, e.g. RequestFailure(InvalidCreds());
+//  - Parameters:
+	// - CategoryName: The name of a category previously defined by ONLINE_ERROR_CATEGORY
+	// - Name: Code-facing name
+	// - ErrorCodeValue: Unique identifiable integer that will encode this error string when displaying the condensed version to end user
+	// - ErrorMessage: Message displayed in the log when this error is printed
+	// - ErrorText: End-User-facing message shown when this error occurs
+	// todo: Do we need to add a value in here to have a "System" other than "Engine"? Maybe we can tie this to the category?
+#define UE_ONLINE_ERROR_INTERNAL(CategoryName, Name, ErrorCodeValue, ErrorMessage, ErrorText) \
+	inline UE::Online::FOnlineError Name(const TOptional<UE::Online::FOnlineError>& Inner = TOptional<UE::Online::FOnlineError>()) \
+		{ \
+			TSharedPtr<UE::Online::FOnlineError, ESPMode::ThreadSafe> InnerPtr; \
+			if (Inner) \
+			{ \
+				InnerPtr = MakeShared<UE::Online::FOnlineError, ESPMode::ThreadSafe>(Inner.GetValue()); \
+			} \
+			return UE::Online::FOnlineError( ErrorCode::Create(UE::Online::Errors::ErrorCode::Category::CategoryName##_System, UE::Online::Errors::ErrorCode::Category::CategoryName, ErrorCodeValue), MakeShared<UE::Online::FOnlineErrorDetails, ESPMode::ThreadSafe>(ErrorMessage, ErrorText), InnerPtr); \
+		}
+
+#define UE_ONLINE_ERROR(CategoryName, Name, ErrorCodeValue, ErrorMessage, ErrorText) \
+    namespace ErrorCode { namespace CategoryName { static constexpr ErrorCodeType Name = Create(Category::CategoryName, ErrorCodeValue); } } \
+	namespace CategoryName { \
+		UE_ONLINE_ERROR_INTERNAL(CategoryName, Name, ErrorCodeValue, ErrorMessage, ErrorText)\
+	}
+
+// Same as above except no namespace around the error, e.g. Errors::NotImplemented instead of Errors::OnlineServices::NotImplemented
+#define UE_ONLINE_ERROR_COMMON(CategoryName, Name, ErrorCodeValue, ErrorMessage, ErrorText) \
+    namespace ErrorCode { namespace CategoryName { static constexpr ErrorCodeType Name = Create(Category::CategoryName, ErrorCodeValue); } } \
+	UE_ONLINE_ERROR_INTERNAL(CategoryName, Name, ErrorCodeValue, ErrorMessage, ErrorText)
+
+} /* namespace UE::Online::Errors */
\ No newline at end of file
diff --git a/Engine/Plugins/Online/OnlineServicesEOS/Source/Private/Online/PresenceEOS.cpp b/Engine/Plugins/Online/OnlineServicesEOS/Source/Private/Online/PresenceEOS.cpp
new file mode 100644
--- /dev/null
+++ b/Engine/Plugins/Online/OnlineServicesEOS/Source/Private/Online/PresenceEOS.cpp
@@ -0,0 +1,646 @@
+// Copyright Epic Games, Inc. All Rights Reserved.
+
+#include "Online/PresenceEOS.h"
+#include "Online/OnlineAsyncOp.h"
+#include "Online/OnlineIdEOS.h"
+#include "Online/OnlineServicesEOS.h"
+#include "Online/OnlineServicesEOSTypes.h"
+#include "Online/AuthEOS.h"
+#include "Online/OnlineErrorEOSGS.h"
+
+#include "eos_presence.h"
+
+namespace UE::Online {
+
+static inline EUserPresenceStatus ToEPresenceState(EOS_Presence_EStatus InStatus)
+{
+	switch (InStatus)
+	{
+		case EOS_Presence_EStatus::EOS_PS_Online:
+		{
+			return EUserPresenceStatus::Online;
+		}
+		case EOS_Presence_EStatus::EOS_PS_Away:
+		{
+			return EUserPresenceStatus::Away;
+		}
+		case EOS_Presence_EStatus::EOS_PS_ExtendedAway:
+		{
+			return EUserPresenceStatus::ExtendedAway;
+		}
+		case EOS_Presence_EStatus::EOS_PS_DoNotDisturb:
+		{
+			return EUserPresenceStatus::DoNotDisturb;
+		}
+	}
+	return EUserPresenceStatus::Offline;
+}
+
+static inline EOS_Presence_EStatus ToEOS_Presence_EStatus(EUserPresenceStatus InState)
+{
+	switch (InState)
+	{
+	case EUserPresenceStatus::Online:
+	{
+		return EOS_Presence_EStatus::EOS_PS_Online;
+	}
+	case EUserPresenceStatus::Away:
+	{
+		return EOS_Presence_EStatus::EOS_PS_Away;
+	}
+	case EUserPresenceStatus::ExtendedAway:
+	{
+		return EOS_Presence_EStatus::EOS_PS_ExtendedAway;
+	}
+	case EUserPresenceStatus::DoNotDisturb:
+	{
+		return EOS_Presence_EStatus::EOS_PS_DoNotDisturb;
+	}
+	}
+	return EOS_Presence_EStatus::EOS_PS_Offline;
+}
+
+FPresenceEOS::FPresenceEOS(FOnlineServicesEOS& InServices)
+	: FPresenceCommon(InServices)
+{
+}
+
+void FPresenceEOS::Initialize()
+{
+	FPresenceCommon::Initialize();
+
+	PresenceHandle = EOS_Platform_GetPresenceInterface(static_cast<FOnlineServicesEOS&>(GetServices()).GetEOSPlatformHandle());
+	check(PresenceHandle != nullptr);
+
+	// Register for friend updates
+	EOS_Presence_AddNotifyOnPresenceChangedOptions Options = { };
+	Options.ApiVersion = EOS_PRESENCE_ADDNOTIFYONPRESENCECHANGED_API_LATEST;
+	NotifyPresenceChangedNotificationId = EOS_Presence_AddNotifyOnPresenceChanged(PresenceHandle, &Options, this, [](const EOS_Presence_PresenceChangedCallbackInfo* Data)
+	{
+		FPresenceEOS* This = reinterpret_cast<FPresenceEOS*>(Data->ClientData);
+		const FAccountId LocalAccountId = FindAccountIdChecked(Data->LocalUserId);
+
+		This->Services.Get<FAuthEOS>()->ResolveAccountId(LocalAccountId, Data->PresenceUserId)
+		.Next([This, LocalAccountId](const FAccountId& PresenceAccountId)
+		{
+			UE_LOG(LogTemp, Warning, TEXT("OnEOSPresenceUpdate: LocalAccountId=[%s] PresenceAccountId=[%s]"), *ToLogString(LocalAccountId), *ToLogString(PresenceAccountId));
+			This->UpdateUserPresence(LocalAccountId, PresenceAccountId);
+		});
+	});
+}
+
+void FPresenceEOS::PreShutdown()
+{
+	EOS_Presence_RemoveNotifyOnPresenceChanged(PresenceHandle, NotifyPresenceChangedNotificationId);
+}
+
+TOnlineAsyncOpHandle<FQueryPresence> FPresenceEOS::QueryPresence(FQueryPresence::Params&& InParams)
+{
+	TOnlineAsyncOpRef<FQueryPresence> Op = GetJoinableOp<FQueryPresence>(MoveTemp(InParams));
+	if (!Op->IsReady())
+	{
+		// Initialize
+		const FQueryPresence::Params& Params = Op->GetParams();
+		if (!Services.Get<FAuthEOS>()->IsLoggedIn(Params.LocalAccountId))
+		{
+			Op->SetError(Errors::NotLoggedIn());
+			return Op->GetHandle();
+		}
+		EOS_EpicAccountId TargetUserEasId = GetEpicAccountId(Params.TargetAccountId);
+		if (!EOS_EpicAccountId_IsValid(TargetUserEasId))
+		{
+			Op->SetError(Errors::NotLoggedIn());
+			return Op->GetHandle();
+		}
+
+		// TODO:  If we try to query a local user's presence, is that an error, should we return the cached state, should we still ask EOS?
+		const bool bIsLocalUser = Services.Get<FAuthEOS>()->IsLoggedIn(Params.TargetAccountId);
+		if (bIsLocalUser)
+		{
+			Op->SetError(Errors::CannotQueryLocalUsers()); 
+			return Op->GetHandle();
+		}
+
+		Op->Then([this](TOnlineAsyncOp<FQueryPresence>& InAsyncOp, TPromise<const EOS_Presence_QueryPresenceCallbackInfo*>&& Promise)
+			{
+				const FQueryPresence::Params& Params = InAsyncOp.GetParams();
+				EOS_Presence_QueryPresenceOptions QueryPresenceOptions = { };
+				QueryPresenceOptions.ApiVersion = EOS_PRESENCE_QUERYPRESENCE_API_LATEST;
+				QueryPresenceOptions.LocalUserId = GetEpicAccountIdChecked(Params.LocalAccountId);
+				QueryPresenceOptions.TargetUserId = GetEpicAccountIdChecked(Params.TargetAccountId);
+
+				EOS_Async(EOS_Presence_QueryPresence, PresenceHandle, QueryPresenceOptions, MoveTemp(Promise));
+			})
+			.Then([this](TOnlineAsyncOp<FQueryPresence>& InAsyncOp, const EOS_Presence_QueryPresenceCallbackInfo* Data) mutable
+			{
+				UE_LOG(LogTemp, Warning, TEXT("QueryPresenceResult: [%s]"), *LexToString(Data->ResultCode));
+
+				if (Data->ResultCode == EOS_EResult::EOS_Success)
+				{
+					const FQueryPresence::Params& Params = InAsyncOp.GetParams();
+					UpdateUserPresence(Params.LocalAccountId, Params.TargetAccountId);
+					FQueryPresence::Result Result = { FindOrCreatePresence(Params.LocalAccountId, Params.TargetAccountId) };
+					InAsyncOp.SetResult(MoveTemp(Result));
+				}
+				else
+				{
+					InAsyncOp.SetError(Errors::FromEOSResult(Data->ResultCode));
+				}
+			})
+			.Enqueue(GetSerialQueue());
+	}
+	return Op->GetHandle();
+}
+
+TOnlineResult<FGetCachedPresence> FPresenceEOS::GetCachedPresence(FGetCachedPresence::Params&& Params)
+{
+	if (TMap<FAccountId, TSharedRef<FUserPresence>>* PresenceList = PresenceLists.Find(Params.LocalAccountId))
+	{
+		TSharedRef<FUserPresence>* PresencePtr = PresenceList->Find(Params.TargetAccountId);
+		if (PresencePtr)
+		{
+			FGetCachedPresence::Result Result = { *PresencePtr };
+			return TOnlineResult<FGetCachedPresence>(MoveTemp(Result));
+		}
+	}
+	return TOnlineResult<FGetCachedPresence>(Errors::NotFound()); 
+}
+
+
+TOnlineAsyncOpHandle<FUpdatePresence> FPresenceEOS::UpdatePresence(FUpdatePresence::Params&& InParams)
+{
+	TOnlineAsyncOpRef<FUpdatePresence> Op = GetOp<FUpdatePresence>(MoveTemp(InParams));
+	const FUpdatePresence::Params& Params = Op->GetParams();
+	if (!Services.Get<FAuthEOS>()->IsLoggedIn(Params.LocalAccountId))
+	{
+		Op->SetError(Errors::NotLoggedIn());
+		return Op->GetHandle();
+	}
+
+	Op->Then([this](TOnlineAsyncOp<FUpdatePresence>& InAsyncOp, TPromise<const EOS_Presence_SetPresenceCallbackInfo*>&& Promise) mutable
+	{
+		const FUpdatePresence::Params& Params = InAsyncOp.GetParams();
+		EOS_HPresenceModification ChangeHandle = nullptr;
+		EOS_Presence_CreatePresenceModificationOptions Options = { };
+		Options.ApiVersion = EOS_PRESENCE_CREATEPRESENCEMODIFICATION_API_LATEST;
+		Options.LocalUserId = GetEpicAccountIdChecked(Params.LocalAccountId);
+		EOS_EResult CreatePresenceModificationResult = EOS_Presence_CreatePresenceModification(PresenceHandle, &Options, &ChangeHandle);
+		if (CreatePresenceModificationResult == EOS_EResult::EOS_Success)
+		{
+			check(ChangeHandle != nullptr);
+
+			// State
+			EOS_PresenceModification_SetStatusOptions StatusOptions = { };
+			StatusOptions.ApiVersion = EOS_PRESENCEMODIFICATION_SETSTATUS_API_LATEST;
+			StatusOptions.Status = ToEOS_Presence_EStatus(Params.Presence->Status);
+			EOS_EResult SetStatusResult = EOS_PresenceModification_SetStatus(ChangeHandle, &StatusOptions);
+			if (SetStatusResult != EOS_EResult::EOS_Success)
+			{
+				UE_LOG(LogTemp, Warning, TEXT("UpdatePresence: EOS_PresenceModification_SetStatus failed with result %s"), *LexToString(SetStatusResult));
+				InAsyncOp.SetError(Errors::FromEOSResult(SetStatusResult));
+				Promise.EmplaceValue();
+				return;
+			}
+
+			// Raw rich text
+			// Convert the status string as the rich text string
+			EOS_PresenceModification_SetRawRichTextOptions RawRichTextOptions = { };
+			RawRichTextOptions.ApiVersion = EOS_PRESENCEMODIFICATION_SETRAWRICHTEXT_API_LATEST;
+
+			const FTCHARToUTF8 Utf8RawRichText(*Params.Presence->StatusString);
+			RawRichTextOptions.RichText = Utf8RawRichText.Get();
+
+			EOS_EResult SetRawRichTextResult = EOS_PresenceModification_SetRawRichText(ChangeHandle, &RawRichTextOptions);
+			if (SetRawRichTextResult != EOS_EResult::EOS_Success)
+			{
+				UE_LOG(LogTemp, Warning, TEXT("UpdatePresence: EOS_PresenceModification_SetRawRichText failed with result %s"), *LexToString(SetRawRichTextResult));
+				InAsyncOp.SetError(Errors::FromEOSResult(SetRawRichTextResult));
+				Promise.EmplaceValue();
+				return;
+			}
+
+			// EOS needs to be specific on which fields are removed and which aren't, so grab the last presence and check which are deletions
+			TArray<FString> RemovedProperties;
+			TSharedRef<FUserPresence> LastUserPresence = FindOrCreatePresence(Params.LocalAccountId, Params.LocalAccountId);
+			for (const TPair<FString, FPresenceVariant>& Pair : LastUserPresence->Properties)
+			{
+				if(!Params.Presence->Properties.Contains(Pair.Key))
+				{
+					RemovedProperties.Add(Pair.Key);
+				}
+			}
+
+			// Removed fields
+			if (RemovedProperties.Num() > 0)
+			{
+				// EOS_PresenceModification_DeleteData
+				TArray<FTCHARToUTF8, TInlineAllocator<EOS_PRESENCE_DATA_MAX_KEYS>> Utf8Strings;
+				TArray<EOS_PresenceModification_DataRecordId, TInlineAllocator<EOS_PRESENCE_DATA_MAX_KEYS>> RecordIds;
+				int32 CurrentIndex = 0;
+				for (const FString& RemovedProperty : RemovedProperties)
+				{
+					const FTCHARToUTF8& Utf8Key = Utf8Strings.Emplace_GetRef(*RemovedProperty);
+
+					EOS_PresenceModification_DataRecordId& RecordId = RecordIds.Emplace_GetRef();
+					RecordId.ApiVersion = EOS_PRESENCEMODIFICATION_DATARECORDID_API_LATEST;
+					RecordId.Key = Utf8Key.Get();
+
+					UE_LOG(LogTemp, Warning, TEXT("UpdatePresence: Removing field %s"), *RemovedProperty); // Temp logging
+				}
+
+				EOS_PresenceModification_DeleteDataOptions DataOptions = { };
+				DataOptions.ApiVersion = EOS_PRESENCEMODIFICATION_DELETEDATA_API_LATEST;
+				DataOptions.RecordsCount = RecordIds.Num();
+				DataOptions.Records = RecordIds.GetData();
+				EOS_EResult DeleteDataResult = EOS_PresenceModification_DeleteData(ChangeHandle, &DataOptions);
+				if (DeleteDataResult != EOS_EResult::EOS_Success)
+				{
+					UE_LOG(LogTemp, Warning, TEXT("UpdatePresence: EOS_PresenceModification_DeleteDataOptions failed with result %s"), *LexToString(DeleteDataResult));
+					InAsyncOp.SetError(Errors::FromEOSResult(DeleteDataResult));
+					Promise.EmplaceValue();
+					return;
+				}
+			}
+
+			// Added/Updated fields
+			if (Params.Presence->Properties.Num() > 0)
+			{
+				if (Params.Presence->Properties.Num() > EOS_PRESENCE_DATA_MAX_KEYS)
+				{
+					// TODO: Move this check higher.  Needs to take into account number of present fields (not just ones updated) and removed fields.
+					UE_LOG(LogTemp, Warning, TEXT("UpdatePresence: Too many presence keys.  %u/%u"), Params.Presence->Properties.Num(), EOS_PRESENCE_DATA_MAX_KEYS);
+					InAsyncOp.SetError(Errors::InvalidParams());
+					Promise.EmplaceValue();
+					return;
+				}
+				TArray<FTCHARToUTF8, TInlineAllocator<EOS_PRESENCE_DATA_MAX_KEYS * 2>> Utf8Strings;
+				TArray<EOS_Presence_DataRecord, TInlineAllocator<EOS_PRESENCE_DATA_MAX_KEYS>> Records;
+				int32 CurrentIndex = 0;
+				for (const TPair<FString, FPresenceVariant>& UpdatedProperty : Params.Presence->Properties)
+				{
+					const FTCHARToUTF8& Utf8Key = Utf8Strings.Emplace_GetRef(*UpdatedProperty.Key);
+					const FTCHARToUTF8& Utf8Value = Utf8Strings.Emplace_GetRef(*UpdatedProperty.Value); // TODO: Better serialization
+
+					EOS_Presence_DataRecord& Record = Records.Emplace_GetRef();
+					Record.ApiVersion = EOS_PRESENCE_DATARECORD_API_LATEST;
+					Record.Key = Utf8Key.Get();
+					Record.Value = Utf8Value.Get();
+					UE_LOG(LogTemp, Warning, TEXT("UpdatePresence: Set field [%s] to [%s]"), *UpdatedProperty.Key, *UpdatedProperty.Value); // Temp logging
+				}
+
+				EOS_PresenceModification_SetDataOptions DataOptions = { };
+				DataOptions.ApiVersion = EOS_PRESENCEMODIFICATION_SETDATA_API_LATEST;
+				DataOptions.RecordsCount = Records.Num();
+				DataOptions.Records = Records.GetData();
+				EOS_EResult SetDataResult = EOS_PresenceModification_SetData(ChangeHandle, &DataOptions);
+				if (SetDataResult != EOS_EResult::EOS_Success)
+				{
+					UE_LOG(LogTemp, Warning, TEXT("UpdatePresence: EOS_PresenceModification_SetData failed with result %s"), *LexToString(SetDataResult));
+					InAsyncOp.SetError(Errors::FromEOSResult(SetDataResult));
+					Promise.EmplaceValue();
+					return;
+				}
+			}
+
+			EOS_Presence_SetPresenceOptions SetPresenceOptions = { };
+			SetPresenceOptions.ApiVersion = EOS_PRESENCE_SETPRESENCE_API_LATEST;
+			SetPresenceOptions.LocalUserId = GetEpicAccountIdChecked(Params.LocalAccountId);
+			SetPresenceOptions.PresenceModificationHandle = ChangeHandle;
+
+			EOS_Async(EOS_Presence_SetPresence, PresenceHandle, SetPresenceOptions, MoveTemp(Promise));
+			EOS_PresenceModification_Release(ChangeHandle);
+			return;
+		}
+		else
+		{
+			InAsyncOp.SetError(Errors::FromEOSResult(CreatePresenceModificationResult));
+		}
+		Promise.EmplaceValue();
+	})
+	.Then([this](TOnlineAsyncOp<FUpdatePresence>& InAsyncOp, const EOS_Presence_SetPresenceCallbackInfo* Data) mutable
+	{
+		UE_LOG(LogTemp, Warning, TEXT("SetPresenceResult: [%s]"), *LexToString(Data->ResultCode));
+
+		if (Data->ResultCode == EOS_EResult::EOS_Success)
+		{
+			// Update local presence
+			const FUpdatePresence::Params& Params = InAsyncOp.GetParams();
+			TSharedRef<FUserPresence> LocalUserPresence = FindOrCreatePresence(Params.LocalAccountId, Params.LocalAccountId);
+			*LocalUserPresence = *Params.Presence;
+
+			InAsyncOp.SetResult(FUpdatePresence::Result());
+		}
+		else
+		{
+			InAsyncOp.SetError(Errors::FromEOSResult(Data->ResultCode));
+		}
+	})
+	.Enqueue(GetSerialQueue());
+
+	return Op->GetHandle();
+}
+
+TOnlineAsyncOpHandle<FPartialUpdatePresence> FPresenceEOS::PartialUpdatePresence(FPartialUpdatePresence::Params&& InParams)
+{
+	// TODO: Validate params
+	// EOS_PRESENCE_DATA_MAX_KEYS - number of total keys.  Compare proposed with existing with pending ops?  Include removed keys!
+	// EOS_PRESENCE_DATA_MAX_KEY_LENGTH - length of each key.  Compare updated with this.
+	// EOS_PRESENCE_DATA_MAX_VALUE_LENGTH - length of each value.  Compare updated with this.
+	// EOS_PRESENCE_RICH_TEXT_MAX_VALUE_LENGTH - length of status. Compare updated with this.
+
+	TOnlineAsyncOpRef<FPartialUpdatePresence> Op = GetMergeableOp<FPartialUpdatePresence>(MoveTemp(InParams));
+	if (!Op->IsReady())
+	{
+		// Initialize
+		const FPartialUpdatePresence::Params& Params = Op->GetParams();
+		if (!Services.Get<FAuthEOS>()->IsLoggedIn(Params.LocalAccountId))
+		{
+			Op->SetError(Errors::NotLoggedIn());
+			return Op->GetHandle();
+		}
+
+		// Don't cache anything from Params as they could be modified by another merge in the meanwhile.
+		Op->Then([this](TOnlineAsyncOp<FPartialUpdatePresence>& InAsyncOp, TPromise<const EOS_Presence_SetPresenceCallbackInfo*>&& Promise) mutable
+			{
+				const FPartialUpdatePresence::Params& Params = InAsyncOp.GetParams();
+				EOS_HPresenceModification ChangeHandle = nullptr;
+				EOS_Presence_CreatePresenceModificationOptions Options = { };
+				Options.ApiVersion = EOS_PRESENCE_CREATEPRESENCEMODIFICATION_API_LATEST;
+				Options.LocalUserId = GetEpicAccountIdChecked(Params.LocalAccountId);
+				EOS_EResult CreatePresenceModificationResult = EOS_Presence_CreatePresenceModification(PresenceHandle, &Options, &ChangeHandle);
+				if (CreatePresenceModificationResult == EOS_EResult::EOS_Success)
+				{
+					check(ChangeHandle != nullptr);
+
+					// State
+					if (Params.Mutations.Status.IsSet())
+					{
+						EOS_PresenceModification_SetStatusOptions StatusOptions = { };
+						StatusOptions.ApiVersion = EOS_PRESENCEMODIFICATION_SETSTATUS_API_LATEST;
+						StatusOptions.Status = ToEOS_Presence_EStatus(Params.Mutations.Status.GetValue());
+						EOS_EResult SetStatusResult = EOS_PresenceModification_SetStatus(ChangeHandle, &StatusOptions);
+						if (SetStatusResult != EOS_EResult::EOS_Success)
+						{
+							UE_LOG(LogTemp, Warning, TEXT("UpdatePresence: EOS_PresenceModification_SetStatus failed with result %s"), *LexToString(SetStatusResult));
+							InAsyncOp.SetError(Errors::FromEOSResult(SetStatusResult));
+							Promise.EmplaceValue();
+							return;
+						}
+						else
+						{
+							UE_LOG(LogTemp, Warning, TEXT("UpdatePresence: Status set to %u"), (uint8)Params.Mutations.Status.GetValue()); // Temp logging
+						}
+					}
+
+					// Raw rich text
+					if (Params.Mutations.StatusString.IsSet())
+					{
+						// Convert the status string as the rich text string
+						EOS_PresenceModification_SetRawRichTextOptions RawRichTextOptions = { };
+						RawRichTextOptions.ApiVersion = EOS_PRESENCEMODIFICATION_SETRAWRICHTEXT_API_LATEST;
+
+						const FTCHARToUTF8 Utf8RawRichText(*Params.Mutations.StatusString.GetValue());
+						RawRichTextOptions.RichText = Utf8RawRichText.Get();
+
+						EOS_EResult SetRawRichTextResult = EOS_PresenceModification_SetRawRichText(ChangeHandle, &RawRichTextOptions);
+						if (SetRawRichTextResult != EOS_EResult::EOS_Success)
+						{
+							UE_LOG(LogTemp, Warning, TEXT("UpdatePresence: EOS_PresenceModification_SetRawRichText failed with result %s"), *LexToString(SetRawRichTextResult));
+							InAsyncOp.SetError(Errors::FromEOSResult(SetRawRichTextResult));
+							Promise.EmplaceValue();
+							return;
+						}
+						else
+						{
+							UE_LOG(LogTemp, Warning, TEXT("UpdatePresence: RichText set to %s"), *Params.Mutations.StatusString.GetValue()); // Temp logging
+						}
+					}
+					// Removed fields
+					if (Params.Mutations.RemovedProperties.Num() > 0)
+					{
+						// EOS_PresenceModification_DeleteData
+						TArray<FTCHARToUTF8, TInlineAllocator<EOS_PRESENCE_DATA_MAX_KEYS>> Utf8Strings;
+						TArray<EOS_PresenceModification_DataRecordId, TInlineAllocator<EOS_PRESENCE_DATA_MAX_KEYS>> RecordIds;
+						int32 CurrentIndex = 0;
+						for (const FString& RemovedProperty : Params.Mutations.RemovedProperties)
+						{
+							const FTCHARToUTF8& Utf8Key = Utf8Strings.Emplace_GetRef(*RemovedProperty);
+
+							EOS_PresenceModification_DataRecordId& RecordId = RecordIds.Emplace_GetRef();
+							RecordId.ApiVersion = EOS_PRESENCEMODIFICATION_DATARECORDID_API_LATEST;
+							RecordId.Key = Utf8Key.Get();
+
+							UE_LOG(LogTemp, Warning, TEXT("UpdatePresence: Removing field %s"), *RemovedProperty); // Temp logging
+						}
+
+						EOS_PresenceModification_DeleteDataOptions DataOptions = { };
+						DataOptions.ApiVersion = EOS_PRESENCEMODIFICATION_DELETEDATA_API_LATEST;
+						DataOptions.RecordsCount = RecordIds.Num();
+						DataOptions.Records = RecordIds.GetData();
+						EOS_EResult DeleteDataResult = EOS_PresenceModification_DeleteData(ChangeHandle, &DataOptions);
+						if (DeleteDataResult != EOS_EResult::EOS_Success)
+						{
+							UE_LOG(LogTemp, Warning, TEXT("UpdatePresence: EOS_PresenceModification_DeleteDataOptions failed with result %s"), *LexToString(DeleteDataResult));
+							InAsyncOp.SetError(Errors::FromEOSResult(DeleteDataResult));
+							Promise.EmplaceValue();
+							return;
+						}
+					}
+
+					// Added/Updated fields
+					if (Params.Mutations.UpdatedProperties.Num() > 0)
+					{
+						if (Params.Mutations.UpdatedProperties.Num() > EOS_PRESENCE_DATA_MAX_KEYS)
+						{
+							// TODO: Move this check higher.  Needs to take into account number of present fields (not just ones updated) and removed fields.
+							UE_LOG(LogTemp, Warning, TEXT("UpdatePresence: Too many presence keys.  %u/%u"), Params.Mutations.UpdatedProperties.Num(), EOS_PRESENCE_DATA_MAX_KEYS);
+							InAsyncOp.SetError(Errors::InvalidParams());
+							Promise.EmplaceValue();
+							return;
+						}
+						TArray<FTCHARToUTF8, TInlineAllocator<EOS_PRESENCE_DATA_MAX_KEYS * 2>> Utf8Strings;
+						TArray<EOS_Presence_DataRecord, TInlineAllocator<EOS_PRESENCE_DATA_MAX_KEYS>> Records;
+						int32 CurrentIndex = 0;
+						for (const TPair<FString, FPresenceVariant>& UpdatedProperty : Params.Mutations.UpdatedProperties)
+						{
+							const FTCHARToUTF8& Utf8Key = Utf8Strings.Emplace_GetRef(*UpdatedProperty.Key);
+							const FTCHARToUTF8& Utf8Value = Utf8Strings.Emplace_GetRef(*UpdatedProperty.Value); // TODO: Better serialization
+
+							EOS_Presence_DataRecord& Record = Records.Emplace_GetRef();
+							Record.ApiVersion = EOS_PRESENCE_DATARECORD_API_LATEST;
+							Record.Key = Utf8Key.Get();
+							Record.Value = Utf8Value.Get();
+							UE_LOG(LogTemp, Warning, TEXT("UpdatePresence: Set field [%s] to [%s]"), *UpdatedProperty.Key, *UpdatedProperty.Value); // Temp logging
+						}
+
+						EOS_PresenceModification_SetDataOptions DataOptions = { };
+						DataOptions.ApiVersion = EOS_PRESENCEMODIFICATION_SETDATA_API_LATEST;
+						DataOptions.RecordsCount = Records.Num();
+						DataOptions.Records = Records.GetData();
+						EOS_EResult SetDataResult = EOS_PresenceModification_SetData(ChangeHandle, &DataOptions);
+						if (SetDataResult != EOS_EResult::EOS_Success)
+						{
+							UE_LOG(LogTemp, Warning, TEXT("UpdatePresence: EOS_PresenceModification_SetData failed with result %s"), *LexToString(SetDataResult));
+							InAsyncOp.SetError(Errors::FromEOSResult(SetDataResult));
+							Promise.EmplaceValue();
+							return;
+						}
+					}
+
+					EOS_Presence_SetPresenceOptions SetPresenceOptions = { };
+					SetPresenceOptions.ApiVersion = EOS_PRESENCE_SETPRESENCE_API_LATEST;
+					SetPresenceOptions.LocalUserId = GetEpicAccountIdChecked(Params.LocalAccountId);
+					SetPresenceOptions.PresenceModificationHandle = ChangeHandle;
+					EOS_Async(EOS_Presence_SetPresence, PresenceHandle, SetPresenceOptions, MoveTemp(Promise));
+					EOS_PresenceModification_Release(ChangeHandle);
+					return;
+				}
+				else
+				{
+					InAsyncOp.SetError(Errors::FromEOSResult(CreatePresenceModificationResult));
+				}
+				Promise.EmplaceValue();
+			})
+			.Then([this](TOnlineAsyncOp<FPartialUpdatePresence>& InAsyncOp, const EOS_Presence_SetPresenceCallbackInfo* Data) mutable
+			{
+				UE_LOG(LogTemp, Warning, TEXT("SetPresenceResult: [%s]"), *LexToString(Data->ResultCode));
+
+				if (Data->ResultCode == EOS_EResult::EOS_Success)
+				{
+					// Update local presence
+					const FPartialUpdatePresence::Params& Params = InAsyncOp.GetParams();
+					TSharedRef<FUserPresence> LocalUserPresence = FindOrCreatePresence(Params.LocalAccountId, Params.LocalAccountId);
+					if (Params.Mutations.Status.IsSet())
+					{
+						LocalUserPresence->Status = Params.Mutations.Status.GetValue();
+					}
+					if (Params.Mutations.StatusString.IsSet())
+					{
+						LocalUserPresence->StatusString = Params.Mutations.StatusString.GetValue();
+					}
+					if (Params.Mutations.GameStatus.IsSet())
+					{
+						LocalUserPresence->GameStatus = Params.Mutations.GameStatus.GetValue();
+					}
+					if (Params.Mutations.Joinability.IsSet())
+					{
+						LocalUserPresence->Joinability = Params.Mutations.Joinability.GetValue();
+					}
+					for (const FString& RemovedKey : Params.Mutations.RemovedProperties)
+					{
+						LocalUserPresence->Properties.Remove(RemovedKey);
+					}
+					for (const TPair<FString, FPresenceVariant>& UpdatedProperty : Params.Mutations.UpdatedProperties)
+					{
+						LocalUserPresence->Properties.Emplace(UpdatedProperty.Key, UpdatedProperty.Value);
+					}
+
+					InAsyncOp.SetResult(FPartialUpdatePresence::Result());
+				}
+				else
+				{
+					InAsyncOp.SetError(Errors::FromEOSResult(Data->ResultCode));
+				}
+			})
+			.Enqueue(GetSerialQueue());
+	}
+	return Op->GetHandle();
+}
+
+/** Get a user's presence, creating entries if missing */
+TSharedRef<FUserPresence> FPresenceEOS::FindOrCreatePresence(FAccountId LocalAccountId, FAccountId PresenceAccountId)
+{
+	TMap<FAccountId, TSharedRef<FUserPresence>>& LocalUserPresenceList = PresenceLists.FindOrAdd(LocalAccountId);
+	if (const TSharedRef<FUserPresence>* const ExistingPresence = LocalUserPresenceList.Find(PresenceAccountId))
+	{
+		return *ExistingPresence;
+	}
+
+	TSharedRef<FUserPresence> UserPresence = MakeShared<FUserPresence>();
+	UserPresence->AccountId = PresenceAccountId;
+	LocalUserPresenceList.Emplace(PresenceAccountId, UserPresence);
+	return UserPresence;
+}
+
+void FPresenceEOS::UpdateUserPresence(FAccountId LocalAccountId, FAccountId PresenceAccountId)
+{
+	bool bPresenceHasChanged = false;
+	TSharedRef<FUserPresence> UserPresence = FindOrCreatePresence(LocalAccountId, PresenceAccountId);
+	// TODO:  Handle updates for local users.  Don't want to conflict with UpdatePresence calls
+
+	// Get presence from EOS
+	EOS_Presence_Info* PresenceInfo = nullptr;
+	EOS_Presence_CopyPresenceOptions Options = { };
+	Options.ApiVersion = EOS_PRESENCE_COPYPRESENCE_API_LATEST;
+	Options.LocalUserId = GetEpicAccountIdChecked(LocalAccountId);
+	Options.TargetUserId = GetEpicAccountIdChecked(PresenceAccountId);
+	EOS_EResult CopyPresenceResult = EOS_Presence_CopyPresence(PresenceHandle, &Options, &PresenceInfo);
+	if (CopyPresenceResult == EOS_EResult::EOS_Success)
+	{
+		// Convert the presence data to our format
+		EUserPresenceStatus NewPresenceState = ToEPresenceState(PresenceInfo->Status);
+		if (UserPresence->Status != NewPresenceState)
+		{
+			bPresenceHasChanged = true;
+			UserPresence->Status = NewPresenceState;
+		}
+
+		FString NewStatusString = UTF8_TO_TCHAR(PresenceInfo->RichText);
+		if (UserPresence->StatusString != NewStatusString)
+		{
+			bPresenceHasChanged = true;
+			UserPresence->StatusString = MoveTemp(NewStatusString);
+		}
+
+		if (PresenceInfo->Records)
+		{
+			// TODO:  Handle Properties that aren't replicated through presence (eg "ProductId")
+			TArrayView<const EOS_Presence_DataRecord> Records(PresenceInfo->Records, PresenceInfo->RecordsCount);
+			// Detect removals
+			TArray<FString, TInlineAllocator<EOS_PRESENCE_DATA_MAX_KEYS>> RemovedKeys;
+			UserPresence->Properties.GenerateKeyArray(RemovedKeys);
+			for (const EOS_Presence_DataRecord& Record : Records)
+			{
+				FString RecordKey = UTF8_TO_TCHAR(Record.Key);
+				FString RecordValue = UTF8_TO_TCHAR(Record.Value);
+				RemovedKeys.Remove(RecordKey);
+				if (FString* ExistingValue = UserPresence->Properties.Find(RecordKey))
+				{
+					if (*ExistingValue == RecordValue)
+					{
+						continue; // No change
+					}
+				}
+				bPresenceHasChanged = true;
+				UserPresence->Properties.Add(MoveTemp(RecordKey), MoveTemp(RecordValue));
+			}
+			// Any fields that have been removed
+			if (RemovedKeys.Num() > 0)
+			{
+				bPresenceHasChanged = true;
+				for (const FString& RemovedKey : RemovedKeys)
+				{
+					UserPresence->Properties.Remove(RemovedKey);
+				}
+			}
+		}
+		else if (UserPresence->Properties.Num() > 0)
+		{
+			bPresenceHasChanged = true;
+			UserPresence->Properties.Reset();
+		}
+		EOS_Presence_Info_Release(PresenceInfo);
+	}
+	else
+	{
+		UE_LOG(LogTemp, Error, TEXT("UpdateUserPresence: CopyPresence Failed %s"), *LexToString(CopyPresenceResult));
+	}
+
+	if (bPresenceHasChanged)
+	{
+		FPresenceUpdated PresenceUpdatedParams = { LocalAccountId, UserPresence };
+		OnPresenceUpdatedEvent.Broadcast(PresenceUpdatedParams);
+	}
+}
+
+/* UE::Online */ }
diff --git a/Engine/Plugins/Online/OnlineServicesEOS/Source/Public/Online/AuthEOS.h b/Engine/Plugins/Online/OnlineServicesEOS/Source/Public/Online/AuthEOS.h
new file mode 100644
--- /dev/null
+++ b/Engine/Plugins/Online/OnlineServicesEOS/Source/Public/Online/AuthEOS.h
@@ -0,0 +1,52 @@
+// Copyright Epic Games, Inc. All Rights Reserved.
+
+#pragma once
+
+#include "Online/AuthEOSGS.h"
+
+#include "eos_userinfo_types.h"
+
+namespace UE::Online {
+
+class FOnlineServicesEOS;
+
+class ONLINESERVICESEOS_API FAuthEOS : public FAuthEOSGS
+{
+public:
+	using Super = FAuthEOSGS;
+
+	FAuthEOS(FOnlineServicesEOS& InOwningSubsystem);
+	virtual ~FAuthEOS() = default;
+
+	// Begin IOnlineComponent
+	virtual void Initialize() override;
+	// End IOnlineComponent
+
+	// Begin IAuth
+	virtual TOnlineAsyncOpHandle<FAuthLogin> Login(FAuthLogin::Params&& Params) override;
+	virtual TOnlineAsyncOpHandle<FAuthQueryExternalServerAuthTicket> QueryExternalServerAuthTicket(FAuthQueryExternalServerAuthTicket::Params&& Params) override;
+	virtual TOnlineAsyncOpHandle<FAuthQueryExternalAuthToken> QueryExternalAuthToken(FAuthQueryExternalAuthToken::Params&& Params) override;
+	// End IAuth
+
+	// Begin FAuthEOSGS
+	virtual TFuture<FAccountId> ResolveAccountId(const FAccountId& LocalAccountId, const EOS_ProductUserId ProductUserId) override;
+	virtual TFuture<TArray<FAccountId>> ResolveAccountIds(const FAccountId& LocalAccountId, const TArray<EOS_ProductUserId>& ProductUserIds) override;
+	virtual TFunction<TFuture<FAccountId>(FOnlineAsyncOp& InAsyncOp, const EOS_ProductUserId& ProductUserId)> ResolveProductIdFn() override;
+	virtual TFunction<TFuture<TArray<FAccountId>>(FOnlineAsyncOp& InAsyncOp, const TArray<EOS_ProductUserId>& ProductUserIds)> ResolveProductIdsFn() override;
+	// End FAuthEOSGS
+
+	TFuture<FAccountId> ResolveAccountId(const FAccountId& LocalAccountId, const EOS_EpicAccountId EpicAccountId);
+	TFuture<TArray<FAccountId>> ResolveAccountIds(const FAccountId& LocalAccountId, const TArray<EOS_EpicAccountId>& EpicAccountIds);
+	TFunction<TFuture<FAccountId>(FOnlineAsyncOp& InAsyncOp, const EOS_EpicAccountId& EpicAccountId)> ResolveEpicIdFn();
+	TFunction<TFuture<TArray<FAccountId>>(FOnlineAsyncOp& InAsyncOp, const TArray<EOS_EpicAccountId>& EpicAccountIds)> ResolveEpicIdsFn();
+
+protected:
+	void ProcessSuccessfulLogin(TOnlineAsyncOp<FAuthLogin>& InAsyncOp);
+	void OnEASLoginStatusChanged(FAccountId LocalAccountId, ELoginStatus PreviousStatus, ELoginStatus CurrentStatus);
+
+	static FAccountId CreateAccountId(const EOS_EpicAccountId EpicAccountId, const EOS_ProductUserId ProductUserId);
+
+	EOS_HUserInfo UserInfoHandle = nullptr;
+};
+
+/* UE::Online */ }
diff --git a/Engine/Plugins/Online/OnlineServicesEOS/Source/Public/Online/ExternalUIEOS.h b/Engine/Plugins/Online/OnlineServicesEOS/Source/Public/Online/ExternalUIEOS.h
new file mode 100644
--- /dev/null
+++ b/Engine/Plugins/Online/OnlineServicesEOS/Source/Public/Online/ExternalUIEOS.h
@@ -0,0 +1,50 @@
+// Copyright Epic Games, Inc. All Rights Reserved.
+
+#pragma once
+
+#include "CoreMinimal.h"
+#include "Online/ExternalUICommon.h"
+
+#if defined(EOS_PLATFORM_BASE_FILE_NAME)
+#include EOS_PLATFORM_BASE_FILE_NAME
+#endif
+#include "eos_ui_types.h"
+
+namespace UE::Online {
+
+struct FExternalUIEOSConfig
+{
+	/** Is ShowLoginUI enabled? */
+	bool bShowLoginUIEnabled = true;
+};
+
+class FOnlineServicesEOS;
+
+class ONLINESERVICESEOS_API FExternalUIEOS : public FExternalUICommon
+{
+public:
+	using Super = FExternalUICommon;
+
+	FExternalUIEOS(FOnlineServicesEOS& InOwningSubsystem);
+	virtual void Initialize() override;
+	virtual void UpdateConfig() override;
+	virtual void PreShutdown() override;
+	virtual TOnlineAsyncOpHandle<FExternalUIShowLoginUI> ShowLoginUI(FExternalUIShowLoginUI::Params&& Params) override;
+	virtual TOnlineAsyncOpHandle<FExternalUIShowFriendsUI> ShowFriendsUI(FExternalUIShowFriendsUI::Params&& Params) override;
+
+protected:
+	/** Handle to EOS UI */
+	EOS_HUI UIHandle;
+	/** Config */
+	FExternalUIEOSConfig Config;
+};
+
+namespace Meta {
+
+BEGIN_ONLINE_STRUCT_META(FExternalUIEOSConfig)
+	ONLINE_STRUCT_FIELD(FExternalUIEOSConfig, bShowLoginUIEnabled)
+END_ONLINE_STRUCT_META()
+
+/* Meta */ }
+
+/* UE::Online */ }
diff --git a/Engine/Plugins/Online/OnlineServicesEOS/Source/Public/Online/PresenceEOS.h b/Engine/Plugins/Online/OnlineServicesEOS/Source/Public/Online/PresenceEOS.h
new file mode 100644
--- /dev/null
+++ b/Engine/Plugins/Online/OnlineServicesEOS/Source/Public/Online/PresenceEOS.h
@@ -0,0 +1,42 @@
+// Copyright Epic Games, Inc. All Rights Reserved.
+
+#pragma once
+
+#include "CoreMinimal.h"
+#include "Online/PresenceCommon.h"
+
+#if defined(EOS_PLATFORM_BASE_FILE_NAME)
+#include EOS_PLATFORM_BASE_FILE_NAME
+#endif
+#include "eos_presence_types.h"
+
+namespace UE::Online {
+
+class FOnlineServicesEOS;
+	
+class FPresenceEOS : public FPresenceCommon
+{
+public:
+	FPresenceEOS(FOnlineServicesEOS& InServices);
+
+	virtual void Initialize() override;
+	virtual void PreShutdown() override;
+
+	virtual TOnlineAsyncOpHandle<FQueryPresence> QueryPresence(FQueryPresence::Params&& Params) override;
+	virtual TOnlineResult<FGetCachedPresence> GetCachedPresence(FGetCachedPresence::Params&& Params) override;
+	virtual TOnlineAsyncOpHandle<FUpdatePresence> UpdatePresence(FUpdatePresence::Params&& Params) override;
+	virtual TOnlineAsyncOpHandle<FPartialUpdatePresence> PartialUpdatePresence(FPartialUpdatePresence::Params&& Params) override;
+
+protected:
+	/** Get a user's presence, creating entries if missing */
+	TSharedRef<FUserPresence> FindOrCreatePresence(FAccountId LocalAccountId, FAccountId PresenceAccountId);
+	/** Update a user's presence from EOS's current value */
+	void UpdateUserPresence(FAccountId LocalAccountId, FAccountId PresenceAccountId);
+protected:
+	EOS_HPresence PresenceHandle = nullptr;
+
+	TMap<FAccountId, TMap<FAccountId, TSharedRef<FUserPresence>>> PresenceLists;
+	EOS_NotificationId NotifyPresenceChangedNotificationId = 0;
+};
+
+/* UE::Online */ }
diff --git a/Engine/Plugins/Online/OnlineServicesEOSGS/Source/Public/Online/OnlineIdEOSGS.h b/Engine/Plugins/Online/OnlineServicesEOSGS/Source/Public/Online/OnlineIdEOSGS.h
new file mode 100644
--- /dev/null
+++ b/Engine/Plugins/Online/OnlineServicesEOSGS/Source/Public/Online/OnlineIdEOSGS.h
@@ -0,0 +1,66 @@
+// Copyright Epic Games, Inc. All Rights Reserved.
+
+#pragma once
+
+#include "Online/CoreOnline.h"
+#include "Online/OnlineIdCommon.h"
+
+#if defined(EOS_PLATFORM_BASE_FILE_NAME)
+#include EOS_PLATFORM_BASE_FILE_NAME
+#endif
+#include "eos_common.h"
+
+namespace UE::Online {
+
+class IOnlineAccountIdRegistryEOSGS : public IOnlineAccountIdRegistry
+{
+public:
+	virtual FAccountId FindAccountId(EOS_ProductUserId ProductUserId) const = 0;
+	virtual EOS_ProductUserId GetProductUserId(const FAccountId& AccountId) const = 0;
+};
+
+/**
+ * Account id registry specifically for EOS id's which are segmented.
+ */
+class FOnlineAccountIdRegistryEOSGS
+	: public IOnlineAccountIdRegistryEOSGS
+{
+public:
+	FOnlineAccountIdRegistryEOSGS();
+	virtual ~FOnlineAccountIdRegistryEOSGS() = default;
+
+	// Begin IOnlineAccountIdRegistryEOSGS
+	virtual FAccountId FindAccountId(const EOS_ProductUserId ProductUserId) const override;
+	virtual EOS_ProductUserId GetProductUserId(const FAccountId& AccountId) const override;
+	// End IOnlineAccountIdRegistryEOSGS
+
+	// Begin IOnlineAccountIdRegistry
+	virtual FString ToLogString(const FAccountId& AccountId) const override;
+	virtual TArray<uint8> ToReplicationData(const FAccountId& AccountId) const override;
+	virtual FAccountId FromReplicationData(const TArray<uint8>& ReplicationData) override;
+	// End IOnlineAccountIdRegistry
+
+	static IOnlineAccountIdRegistryEOSGS& GetRegistered();
+
+private:
+	// FAuthEOSGS is the only thing that should be able to create PUID-only net ids in this registry, in its resolve methods.
+	friend class FAuthEOSGS;
+	friend class FOnlineServicesEOSGSModule;
+	static FOnlineAccountIdRegistryEOSGS& Get();
+	FAccountId FindOrAddAccountId(const EOS_ProductUserId ProductUserId);
+
+	TOnlineBasicAccountIdRegistry<EOS_ProductUserId> Registry;
+};
+
+EOS_ProductUserId ONLINESERVICESEOSGS_API GetProductUserId(const FAccountId& AccountId);
+EOS_ProductUserId ONLINESERVICESEOSGS_API GetProductUserIdChecked(const FAccountId& AccountId);
+FAccountId ONLINESERVICESEOSGS_API FindAccountId(const EOS_ProductUserId EpicAccountId);
+FAccountId ONLINESERVICESEOSGS_API FindAccountIdChecked(const EOS_ProductUserId EpicAccountId);
+
+template<typename IdType>
+inline bool ValidateOnlineId(const TOnlineId<IdType> OnlineId)
+{
+	return OnlineId.GetOnlineServicesType() == EOnlineServices::Epic && OnlineId.IsValid();
+}
+
+} /* namespace UE::Online */
diff --git a/Engine/Plugins/Online/OnlineServicesNull/Source/Private/Online/StatsNull.cpp b/Engine/Plugins/Online/OnlineServicesNull/Source/Private/Online/StatsNull.cpp
new file mode 100644
--- /dev/null
+++ b/Engine/Plugins/Online/OnlineServicesNull/Source/Private/Online/StatsNull.cpp
@@ -0,0 +1,184 @@
+// Copyright Epic Games, Inc. All Rights Reserved.
+
+#include "Online/StatsNull.h"
+#include "Online/AuthNull.h"
+#include "Online/OnlineServicesNull.h"
+#include "Math/UnrealMathUtility.h"
+
+namespace UE::Online {
+
+FStatsNull::FStatsNull(FOnlineServicesNull& InOwningSubsystem)
+	: Super(InOwningSubsystem)
+{
+}
+
+TOnlineAsyncOpHandle<FUpdateStats> FStatsNull::UpdateStats(FUpdateStats::Params&& Params)
+{
+	TOnlineAsyncOpRef<FUpdateStats> Op = GetOp<FUpdateStats>(MoveTemp(Params));
+
+	if (!Services.Get<FAuthNull>()->IsLoggedIn(Op->GetParams().LocalAccountId))
+	{
+		Op->SetError(Errors::InvalidUser());
+		return Op->GetHandle();
+	}
+
+	Op->Then([this](TOnlineAsyncOp<FUpdateStats>& InAsyncOp) mutable
+	{
+		for (const FUserStats& UpdateUserStats : InAsyncOp.GetParams().UpdateUsersStats)
+		{
+			FUserStats* ExistingUserStats = UsersStats.FindByPredicate(FFindUserStatsByAccountId(UpdateUserStats.AccountId));
+			if (!ExistingUserStats)
+			{
+				ExistingUserStats = &UsersStats.Emplace_GetRef();
+				ExistingUserStats->AccountId = UpdateUserStats.AccountId;
+			}
+
+			TMap<FString, FStatValue>& UserStats = ExistingUserStats->Stats;
+			for (const auto& UpdateUserStatPair : UpdateUserStats.Stats)
+			{
+				if (const FStatDefinition* StatDefinition = GetStatDefinition(UpdateUserStatPair.Key))
+				{
+					if (FStatValue* StatValue = UserStats.Find(UpdateUserStatPair.Key))
+					{
+						switch (StatDefinition->ModifyMethod)
+						{
+						case EStatModifyMethod::Set:
+							*StatValue = UpdateUserStatPair.Value;
+							break;
+						case EStatModifyMethod::Sum:
+							if (StatValue->VariantData.IsType<double>())
+							{
+								StatValue->Set(StatValue->GetDouble() + UpdateUserStatPair.Value.GetDouble());
+							}
+							else if (StatValue->VariantData.IsType<int64>())
+							{
+								StatValue->Set(StatValue->GetInt64() + UpdateUserStatPair.Value.GetInt64());
+							}
+							break;
+						case EStatModifyMethod::Largest:
+							if (StatValue->VariantData.IsType<double>())
+							{
+								StatValue->Set(FMath::Max(StatValue->GetDouble(), UpdateUserStatPair.Value.GetDouble()));
+							}
+							else if (StatValue->VariantData.IsType<int64>())
+							{
+								StatValue->Set(FMath::Max(StatValue->GetInt64(), UpdateUserStatPair.Value.GetInt64()));
+							}
+							break;
+						case EStatModifyMethod::Smallest:
+							if (StatValue->VariantData.IsType<double>())
+							{
+								StatValue->Set(FMath::Min(StatValue->GetDouble(), UpdateUserStatPair.Value.GetDouble()));
+							}
+							else if (StatValue->VariantData.IsType<int64>())
+							{
+								StatValue->Set(FMath::Min(StatValue->GetInt64(), UpdateUserStatPair.Value.GetInt64()));
+							}
+							break;
+						}
+					}
+					else
+					{
+						UserStats.Emplace(UpdateUserStatPair.Key, UpdateUserStatPair.Value);
+					}
+				}
+			}
+
+			CacheUserStats(*ExistingUserStats);
+		}
+
+		InAsyncOp.SetResult({});
+		OnStatsUpdatedEvent.Broadcast(InAsyncOp.GetParams());
+	})
+	.Enqueue(GetSerialQueue());
+
+	return Op->GetHandle();
+}
+
+TOnlineAsyncOpHandle<FQueryStats> FStatsNull::QueryStats(FQueryStats::Params&& Params)
+{
+	TOnlineAsyncOpRef<FQueryStats> Op = GetOp<FQueryStats>(MoveTemp(Params));
+
+	if (!Services.Get<FAuthNull>()->IsLoggedIn(Op->GetParams().LocalAccountId))
+	{
+		Op->SetError(Errors::InvalidUser());
+		return Op->GetHandle();
+	}
+
+	Op->Then([this](TOnlineAsyncOp<FQueryStats>& InAsyncOp)
+	{
+		FQueryStats::Result Result;
+
+		if (FUserStats* ExistingUserStats = UsersStats.FindByPredicate(FFindUserStatsByAccountId(InAsyncOp.GetParams().TargetAccountId)))
+		{
+			Result.Stats = ExistingUserStats->Stats;
+		}
+
+		InAsyncOp.SetResult(MoveTemp(Result));
+	})
+	.Enqueue(GetSerialQueue());
+
+	return Op->GetHandle();
+}
+
+TOnlineAsyncOpHandle<FBatchQueryStats> FStatsNull::BatchQueryStats(FBatchQueryStats::Params&& Params)
+{
+	TOnlineAsyncOpRef<FBatchQueryStats> Op = GetOp<FBatchQueryStats>(MoveTemp(Params));
+	if (!Services.Get<FAuthNull>()->IsLoggedIn(Op->GetParams().LocalAccountId))
+	{
+		Op->SetError(Errors::InvalidUser());
+		return Op->GetHandle();
+	}
+
+	Op->Then([this](TOnlineAsyncOp<FBatchQueryStats>& InAsyncOp)
+	{
+		FBatchQueryStats::Result Result;
+
+		for (const FAccountId& TargetAccountId : InAsyncOp.GetParams().TargetAccountIds)
+		{
+			if (FUserStats* ExistingUserStats = UsersStats.FindByPredicate(FFindUserStatsByAccountId(TargetAccountId)))
+			{
+				Result.UsersStats.Emplace(*ExistingUserStats);
+			}
+		}
+
+		InAsyncOp.SetResult(MoveTemp(Result));
+	})
+	.Enqueue(GetSerialQueue());
+
+	return Op->GetHandle();
+}
+
+#if !UE_BUILD_SHIPPING
+TOnlineAsyncOpHandle<FResetStats> FStatsNull::ResetStats(FResetStats::Params&& Params)
+{
+	TOnlineAsyncOpRef<FResetStats> Op = GetOp<FResetStats>(MoveTemp(Params));
+
+	if (!Services.Get<FAuthNull>()->IsLoggedIn(Op->GetParams().LocalAccountId))
+	{
+		Op->SetError(Errors::InvalidUser());
+		return Op->GetHandle();
+	}
+
+	Op->Then([this](TOnlineAsyncOp<FResetStats>& InAsyncOp)
+	{
+		uint32 Index = UsersStats.IndexOfByPredicate(FFindUserStatsByAccountId(InAsyncOp.GetParams().LocalAccountId));
+		if (Index != INDEX_NONE)
+		{
+			UsersStats.RemoveAt(Index);
+		}
+
+		FFindUserStatsByAccountId FindUserStatsByAccountId(InAsyncOp.GetParams().LocalAccountId);
+
+		UsersStats.RemoveAll(FindUserStatsByAccountId);
+		CachedUsersStats.RemoveAll(FindUserStatsByAccountId);
+
+		InAsyncOp.SetResult({});
+	})
+	.Enqueue(GetSerialQueue());
+
+	return Op->GetHandle();
+}
+#endif // !UE_BUILD_SHIPPING
+
+/* UE::Online */ }
