Added logic to queue EOS presence update callbacks received before login, to perform the related calls after login is complete.

diff --git a/Engine/Plugins/Online/OnlineServicesEOS/Source/Private/Online/PresenceEOS.cpp b/Engine/Plugins/Online/OnlineServicesEOS/Source/Private/Online/PresenceEOS.cpp
--- a/Engine/Plugins/Online/OnlineServicesEOS/Source/Private/Online/PresenceEOS.cpp
+++ b/Engine/Plugins/Online/OnlineServicesEOS/Source/Private/Online/PresenceEOS.cpp
@@ -78,20 +78,59 @@ void FPresenceEOS::Initialize()
 	NotifyPresenceChangedNotificationId = EOS_Presence_AddNotifyOnPresenceChanged(PresenceHandle, &Options, this, [](const EOS_Presence_PresenceChangedCallbackInfo* Data)
 	{
 		FPresenceEOS* This = reinterpret_cast<FPresenceEOS*>(Data->ClientData);
-		const FAccountId LocalAccountId = FindAccountIdChecked(Data->LocalUserId);
+		const FAccountId LocalAccountId = FindAccountId(Data->LocalUserId);
 
-		This->Services.Get<FAuthEOS>()->ResolveAccountId(LocalAccountId, Data->PresenceUserId)
-		.Next([This, LocalAccountId](const FAccountId& PresenceAccountId)
+		if (LocalAccountId.IsValid())
 		{
-			UE_LOG(LogTemp, Warning, TEXT("OnEOSPresenceUpdate: LocalAccountId=[%s] PresenceAccountId=[%s]"), *ToLogString(LocalAccountId), *ToLogString(PresenceAccountId));
-			This->UpdateUserPresence(LocalAccountId, PresenceAccountId);
-		});
+			This->Services.Get<FAuthEOS>()->ResolveAccountId(LocalAccountId, Data->PresenceUserId)
+				.Next([This, LocalAccountId](const FAccountId& PresenceAccountId)
+					{
+						UE_LOG(LogTemp, Verbose, TEXT("OnEOSPresenceUpdate: LocalAccountId=[%s] PresenceAccountId=[%s]"), *ToLogString(LocalAccountId), *ToLogString(PresenceAccountId));
+						This->UpdateUserPresence(LocalAccountId, PresenceAccountId);
+					});
+		}
+		else // In some cases, this delegate will fire before the login process has completed, so we won't be able to find the AccountId just yet. We'll queue that presence update call for after Login has completed
+		{
+			UE_LOG(LogTemp, Verbose, TEXT("OnEOSPresenceUpdate: Account id not found for Epic id [%s]. Will retry after login completes]"), *LexToString(Data->LocalUserId));
+
+			TArray<EOS_EpicAccountId>& PendingPresenceUpdateArray = This->PendingPresenceUpdates.FindOrAdd(Data->LocalUserId);
+			PendingPresenceUpdateArray.Add(Data->PresenceUserId);
+		}
 	});
+
+	LoginStatusChangedHandle = Services.Get<IAuth>()->OnLoginStatusChanged().Add(this, &FPresenceEOS::HandleAuthLoginStatusChanged);
 }
 
 void FPresenceEOS::PreShutdown()
 {
 	EOS_Presence_RemoveNotifyOnPresenceChanged(PresenceHandle, NotifyPresenceChangedNotificationId);
+
+	LoginStatusChangedHandle.Unbind();
+}
+
+void FPresenceEOS::HandleAuthLoginStatusChanged(const FAuthLoginStatusChanged& EventParameters)
+{
+	if (EventParameters.LoginStatus == ELoginStatus::LoggedIn)
+	{
+		const FAccountId LocalAccountId = EventParameters.AccountInfo->AccountId;
+
+		const EOS_EpicAccountId EpicAccountId = GetEpicAccountId(LocalAccountId);
+
+		if (TArray<EOS_EpicAccountId>* PendingPresenceUpdateArray = PendingPresenceUpdates.Find(EpicAccountId))
+		{
+			for (EOS_EpicAccountId& PresenceAccountId : *PendingPresenceUpdateArray)
+			{
+				Services.Get<FAuthEOS>()->ResolveAccountId(LocalAccountId, PresenceAccountId)
+					.Next([this, LocalAccountId](const FAccountId& PresenceAccountId)
+						{
+							UE_LOG(LogTemp, Verbose, TEXT("OnEOSPresenceUpdate: LocalAccountId=[%s] PresenceAccountId=[%s]"), *ToLogString(LocalAccountId), *ToLogString(PresenceAccountId));
+							UpdateUserPresence(LocalAccountId, PresenceAccountId);
+						});
+			}
+
+			PendingPresenceUpdates.Remove(EpicAccountId);
+		}
+	}
 }
 
 TOnlineAsyncOpHandle<FQueryPresence> FPresenceEOS::QueryPresence(FQueryPresence::Params&& InParams)
diff --git a/Engine/Plugins/Online/OnlineServicesEOS/Source/Public/Online/PresenceEOS.h b/Engine/Plugins/Online/OnlineServicesEOS/Source/Public/Online/PresenceEOS.h
--- a/Engine/Plugins/Online/OnlineServicesEOS/Source/Public/Online/PresenceEOS.h
+++ b/Engine/Plugins/Online/OnlineServicesEOS/Source/Public/Online/PresenceEOS.h
@@ -32,9 +32,16 @@ protected:
 	TSharedRef<FUserPresence> FindOrCreatePresence(FAccountId LocalAccountId, FAccountId PresenceAccountId);
 	/** Update a user's presence from EOS's current value */
 	void UpdateUserPresence(FAccountId LocalAccountId, FAccountId PresenceAccountId);
+	/** Performs queued presence updates after a user's login completes */
+	void HandleAuthLoginStatusChanged(const FAuthLoginStatusChanged& EventParameters);
+
 protected:
 	EOS_HPresence PresenceHandle = nullptr;
 
+	/** Login status changed event handle */
+	FOnlineEventDelegateHandle LoginStatusChangedHandle;
+
+	TMap<EOS_EpicAccountId, TArray<EOS_EpicAccountId>> PendingPresenceUpdates;
 	TMap<FAccountId, TMap<FAccountId, TSharedRef<FUserPresence>>> PresenceLists;
 	EOS_NotificationId NotifyPresenceChangedNotificationId = 0;
 };
