Move EOS_EpicAccountId / FAccountId accessors to public
Enable EOS_Connect_QueryProductUserIdMappings to be called from servers

diff --git a/Engine/Plugins/Online/OnlineServicesEOS/Source/Private/Online/AuthEOS.cpp b/Engine/Plugins/Online/OnlineServicesEOS/Source/Private/Online/AuthEOS.cpp
--- a/Engine/Plugins/Online/OnlineServicesEOS/Source/Private/Online/AuthEOS.cpp
+++ b/Engine/Plugins/Online/OnlineServicesEOS/Source/Private/Online/AuthEOS.cpp
@@ -384,6 +384,7 @@ TFuture<TArray<FAccountId>> FAuthEOS::ResolveAccountIds(const FAccountId& LocalA
 	}
 
 	// If we failed to find all the handles, we need to query, which requires a valid LocalAccountId
+	// Note this is unavailable on Dedicated Servers as well, unlike EOS_Connect_QueryProductUserIdMappings
 	if (!ValidateOnlineId(LocalAccountId))
 	{
 		checkNoEntry();
@@ -487,7 +488,7 @@ TFuture<TArray<FAccountId>> FAuthEOS::ResolveAccountIds(const FAccountId& LocalA
 	}
 
 	// If we failed to find all the handles, we need to query, which requires a valid LocalAccountId
-	if (!ValidateOnlineId(LocalAccountId))
+	if (!(IsRunningDedicatedServer() || ValidateOnlineId(LocalAccountId)))
 	{
 		checkNoEntry();
 		return MakeFulfilledPromise<TArray<FAccountId>>().GetFuture();
@@ -499,7 +500,7 @@ TFuture<TArray<FAccountId>> FAuthEOS::ResolveAccountIds(const FAccountId& LocalA
 	EOS_Connect_QueryProductUserIdMappingsOptions Options = {};
 	Options.ApiVersion = 2;
 	UE_EOS_CHECK_API_MISMATCH(EOS_CONNECT_QUERYPRODUCTUSERIDMAPPINGS_API_LATEST, 2);
-	Options.LocalUserId = GetProductUserIdChecked(LocalAccountId);
+	Options.LocalUserId = IsRunningDedicatedServer() ? nullptr : GetProductUserIdChecked(LocalAccountId);
 	Options.ProductUserIds = MissingProductUserIds.GetData();
 	Options.ProductUserIdCount = MissingProductUserIds.Num();
 	
diff --git a/Engine/Plugins/Online/OnlineServicesEOS/Source/Private/Online/OnlineIdEOS.h b/Engine/Plugins/Online/OnlineServicesEOS/Source/Private/Online/OnlineIdEOS.h
--- a/Engine/Plugins/Online/OnlineServicesEOS/Source/Private/Online/OnlineIdEOS.h
+++ b/Engine/Plugins/Online/OnlineServicesEOS/Source/Private/Online/OnlineIdEOS.h
@@ -3,6 +3,7 @@
 #pragma once
 
 #include "Online/OnlineIdEOSGS.h"
+#include "Online/AccountIdEOS.h"
 
 struct EOS_EpicAccountIdDetails;
 struct EOS_ProductUserIdDetails;
@@ -59,9 +60,4 @@ private:
 	TMap<EOS_ProductUserId, FAccountId> ProductUserIdToAccountId; // Map of EOS_ProductUserId to the associated handle.
 };
 
-EOS_EpicAccountId GetEpicAccountId(const FAccountId& AccountId);
-EOS_EpicAccountId GetEpicAccountIdChecked(const FAccountId& AccountId);
-FAccountId FindAccountId(const EOS_EpicAccountId EpicAccountId);
-FAccountId FindAccountIdChecked(const EOS_EpicAccountId EpicAccountId);
-
 } /* namespace UE::Online */
diff --git a/Engine/Plugins/Online/OnlineServicesEOS/Source/Public/Online/AccountIdEOS.h b/Engine/Plugins/Online/OnlineServicesEOS/Source/Public/Online/AccountIdEOS.h
new file mode 100644
--- /dev/null
+++ b/Engine/Plugins/Online/OnlineServicesEOS/Source/Public/Online/AccountIdEOS.h
@@ -0,0 +1,17 @@
+// Copyright Epic Games, Inc. All Rights Reserved.
+
+#pragma once
+
+#include "Online/CoreOnline.h"
+
+typedef struct EOS_EpicAccountIdDetails* EOS_EpicAccountId;
+
+namespace UE::Online
+{
+
+ONLINESERVICESEOS_API EOS_EpicAccountId GetEpicAccountId(const FAccountId& AccountId);
+ONLINESERVICESEOS_API EOS_EpicAccountId GetEpicAccountIdChecked(const FAccountId& AccountId);
+ONLINESERVICESEOS_API FAccountId FindAccountId(const EOS_EpicAccountId EpicAccountId);
+ONLINESERVICESEOS_API FAccountId FindAccountIdChecked(const EOS_EpicAccountId EpicAccountId);
+
+} /* namespace UE::Online */
diff --git a/Engine/Plugins/Online/OnlineServicesEOS/Source/Public/Online/AuthEOS.h b/Engine/Plugins/Online/OnlineServicesEOS/Source/Public/Online/AuthEOS.h
--- a/Engine/Plugins/Online/OnlineServicesEOS/Source/Public/Online/AuthEOS.h
+++ b/Engine/Plugins/Online/OnlineServicesEOS/Source/Public/Online/AuthEOS.h
@@ -49,4 +49,6 @@ protected:
 	EOS_HUserInfo UserInfoHandle = nullptr;
 };
 
+using FAuthEOSPtr = TSharedPtr<FAuthEOS>;
+
 /* UE::Online */ }
