- Add public method for creating an EOS v2 net id

diff --git a/Engine/Plugins/Online/OnlineServicesEOS/Source/Private/Online/OnlineIdEOS.cpp b/Engine/Plugins/Online/OnlineServicesEOS/Source/Private/Online/OnlineIdEOS.cpp
--- a/Engine/Plugins/Online/OnlineServicesEOS/Source/Private/Online/OnlineIdEOS.cpp
+++ b/Engine/Plugins/Online/OnlineServicesEOS/Source/Private/Online/OnlineIdEOS.cpp
@@ -292,4 +292,9 @@ FAccountId FindAccountIdChecked(const EOS_EpicAccountId EpicAccountId)
 	return Result;
 }
 
+FAccountId CreateAccountId(const EOS_EpicAccountId EpicAccountId, const EOS_ProductUserId ProductUserId)
+{
+	return FOnlineAccountIdRegistryEOS::Get().FindOrAddAccountId(EpicAccountId, ProductUserId);
+}
+
 /* UE::Online */ }
diff --git a/Engine/Plugins/Online/OnlineServicesEOS/Source/Public/Online/AccountIdEOS.h b/Engine/Plugins/Online/OnlineServicesEOS/Source/Public/Online/AccountIdEOS.h
--- a/Engine/Plugins/Online/OnlineServicesEOS/Source/Public/Online/AccountIdEOS.h
+++ b/Engine/Plugins/Online/OnlineServicesEOS/Source/Public/Online/AccountIdEOS.h
@@ -13,5 +13,6 @@ ONLINESERVICESEOS_API EOS_EpicAccountId GetEpicAccountId(const FAccountId& Accou
 ONLINESERVICESEOS_API EOS_EpicAccountId GetEpicAccountIdChecked(const FAccountId& AccountId);
 ONLINESERVICESEOS_API FAccountId FindAccountId(const EOS_EpicAccountId EpicAccountId);
 ONLINESERVICESEOS_API FAccountId FindAccountIdChecked(const EOS_EpicAccountId EpicAccountId);
+ONLINESERVICESEOS_API FAccountId CreateAccountId(const EOS_EpicAccountId EpicAccountId, const EOS_ProductUserId ProductUserId);
 
 } /* namespace UE::Online */
