Integrated EOSSDK 1.16 BestDisplayName.

diff --git a/Engine/Plugins/Online/EOSShared/Source/EOSShared/Private/EOSSDKManager.cpp b/Engine/Plugins/Online/EOSShared/Source/EOSShared/Private/EOSSDKManager.cpp
--- a/Engine/Plugins/Online/EOSShared/Source/EOSShared/Private/EOSSDKManager.cpp
+++ b/Engine/Plugins/Online/EOSShared/Source/EOSShared/Private/EOSSDKManager.cpp
@@ -951,6 +951,44 @@ void FEOSSDKManager::LogUserInfo(const EOS_HPlatform Platform, const EOS_EpicAcc
 	{
 		UE_LOG_EOSSDK_INFO("UserInfo (EOS_UserInfo_CopyUserInfo failed: %s)", *LexToString(Result));
 	}
+
+	EOS_UserInfo_CopyBestDisplayNameOptions Options = {};
+	Options.ApiVersion = 1;
+	UE_EOS_CHECK_API_MISMATCH(EOS_USERINFO_COPYBESTDISPLAYNAME_API_LATEST, 1);
+	Options.LocalUserId = LoggedInAccount;
+	Options.TargetUserId = TargetAccount;
+
+	EOS_UserInfo_BestDisplayName* BestDisplayName;
+	Result = EOS_UserInfo_CopyBestDisplayName(UserInfoHandle, &Options, &BestDisplayName);
+
+	if (Result == EOS_EResult::EOS_UserInfo_BestDisplayNameIndeterminate)
+	{
+		EOS_UserInfo_CopyBestDisplayNameWithPlatformOptions WithPlatformOptions = {};
+		WithPlatformOptions.ApiVersion = 1;
+		UE_EOS_CHECK_API_MISMATCH(EOS_USERINFO_COPYBESTDISPLAYNAMEWITHPLATFORM_API_LATEST, 1);
+		WithPlatformOptions.LocalUserId = LoggedInAccount;
+		WithPlatformOptions.TargetUserId = TargetAccount;
+		WithPlatformOptions.TargetPlatformType = EOS_OPT_Epic;
+
+		Result = EOS_UserInfo_CopyBestDisplayNameWithPlatform(UserInfoHandle, &WithPlatformOptions, &BestDisplayName);
+	}
+
+	if (Result == EOS_EResult::EOS_Success)
+	{
+		Indent++;
+		UE_LOG_EOSSDK_INFO("BestDisplayName");
+		Indent++;
+		UE_LOG_EOSSDK_INFO("DisplayName=%hs", BestDisplayName->DisplayName);
+		UE_LOG_EOSSDK_INFO("DisplayNameSanitized=%hs", BestDisplayName->DisplayNameSanitized);
+		UE_LOG_EOSSDK_INFO("Nickname=%hs", BestDisplayName->Nickname);
+		Indent -= 2;
+
+		EOS_UserInfo_BestDisplayName_Release(BestDisplayName);
+	}
+	else
+	{
+		UE_LOG_EOSSDK_INFO("UserInfo (BestDisplayName retrieval failed: %s)", *LexToString(Result));
+	}
 }
 
 void FEOSSDKManager::LogPresenceInfo(const EOS_HPlatform Platform, const EOS_EpicAccountId LoggedInAccount, const EOS_EpicAccountId TargetAccount, int32 Indent) const
diff --git a/Engine/Plugins/Online/OnlineServicesEOS/Source/Private/Online/AuthEOS.cpp b/Engine/Plugins/Online/OnlineServicesEOS/Source/Private/Online/AuthEOS.cpp
--- a/Engine/Plugins/Online/OnlineServicesEOS/Source/Private/Online/AuthEOS.cpp
+++ b/Engine/Plugins/Online/OnlineServicesEOS/Source/Private/Online/AuthEOS.cpp
@@ -195,24 +195,36 @@ TOnlineAsyncOpHandle<FAuthLogin> FAuthEOS::Login(FAuthLogin::Params&& Params)
 		const TSharedRef<FAccountInfoEOS>& AccountInfoEOS = GetOpDataChecked<TSharedRef<FAccountInfoEOS>>(InAsyncOp, UE_ONLINE_AUTH_EOS_ACCOUNT_INFO_KEY_NAME);
 
 		// Get display name
-		EOS_UserInfo_CopyUserInfoOptions Options = { };
-		Options.ApiVersion = 3;
-		UE_EOS_CHECK_API_MISMATCH(EOS_USERINFO_COPYUSERINFO_API_LATEST, 3);
+		EOS_UserInfo_CopyBestDisplayNameOptions Options = {};
+		Options.ApiVersion = 1;
+		UE_EOS_CHECK_API_MISMATCH(EOS_USERINFO_COPYBESTDISPLAYNAME_API_LATEST, 1);
 		Options.LocalUserId = AccountInfoEOS->EpicAccountId;
 		Options.TargetUserId = AccountInfoEOS->EpicAccountId;
 
-		EOS_UserInfo* UserInfo = nullptr;
+		EOS_UserInfo_BestDisplayName* BestDisplayName;
+		EOS_EResult CopyBestDisplayNameResult = EOS_UserInfo_CopyBestDisplayName(UserInfoHandle, &Options, &BestDisplayName);
 
-		EOS_EResult CopyUserInfoResult = EOS_UserInfo_CopyUserInfo(UserInfoHandle, &Options, &UserInfo);
-		if (CopyUserInfoResult == EOS_EResult::EOS_Success)
+		if (CopyBestDisplayNameResult == EOS_EResult::EOS_UserInfo_BestDisplayNameIndeterminate)
 		{
-			AccountInfoEOS->Attributes.Emplace(AccountAttributeData::DisplayName, UTF8_TO_TCHAR(UserInfo->DisplayName));
-			EOS_UserInfo_Release(UserInfo);
+			EOS_UserInfo_CopyBestDisplayNameWithPlatformOptions WithPlatformOptions = {};
+			WithPlatformOptions.ApiVersion = 1;
+			UE_EOS_CHECK_API_MISMATCH(EOS_USERINFO_COPYBESTDISPLAYNAMEWITHPLATFORM_API_LATEST, 1);
+			WithPlatformOptions.LocalUserId = AccountInfoEOS->EpicAccountId;
+			WithPlatformOptions.TargetUserId = AccountInfoEOS->EpicAccountId;
+			WithPlatformOptions.TargetPlatformType = EOS_OPT_Epic;
+
+			CopyBestDisplayNameResult = EOS_UserInfo_CopyBestDisplayNameWithPlatform(UserInfoHandle, &WithPlatformOptions, &BestDisplayName);
+		}
+
+		if (CopyBestDisplayNameResult == EOS_EResult::EOS_Success)
+		{
+			AccountInfoEOS->Attributes.Emplace(AccountAttributeData::DisplayName, UTF8_TO_TCHAR(BestDisplayName->DisplayNameSanitized));
+			EOS_UserInfo_BestDisplayName_Release(BestDisplayName);
 		}
 		else
 		{
-			FOnlineError CopyUserInfoError(Errors::FromEOSResult(CopyUserInfoResult));
-			UE_LOG(LogOnlineServices, Warning, TEXT("[FAuthEOS::Login] Failure: EOS_UserInfo_CopyUserInfo %s"), *CopyUserInfoError.GetLogString());
+			FOnlineError CopyUserInfoError(Errors::FromEOSResult(CopyBestDisplayNameResult));
+			UE_LOG(LogOnlineServices, Warning, TEXT("[FAuthEOS::Login] Failure: EOS_UserInfo_CopyBestDisplayName %s"), *CopyUserInfoError.GetLogString());
 			InAsyncOp.SetError(Errors::Unknown(MoveTemp(CopyUserInfoError)));
 
 			TPromise<void> Promise;
diff --git a/Engine/Plugins/Online/OnlineServicesEOS/Source/Private/Online/UserInfoEOS.cpp b/Engine/Plugins/Online/OnlineServicesEOS/Source/Private/Online/UserInfoEOS.cpp
--- a/Engine/Plugins/Online/OnlineServicesEOS/Source/Private/Online/UserInfoEOS.cpp
+++ b/Engine/Plugins/Online/OnlineServicesEOS/Source/Private/Online/UserInfoEOS.cpp
@@ -99,28 +99,41 @@ TOnlineResult<FGetUserInfo> FUserInfoEOS::GetUserInfo(FGetUserInfo::Params&& Par
 		return TOnlineResult<FGetUserInfo>(Errors::InvalidParams());
 	}
 
-	EOS_UserInfo_CopyUserInfoOptions Options;
-	Options.ApiVersion = 3;
-	UE_EOS_CHECK_API_MISMATCH(EOS_USERINFO_COPYUSERINFO_API_LATEST, 3);
+	EOS_UserInfo_CopyBestDisplayNameOptions Options = {};
+	Options.ApiVersion = 1;
+	UE_EOS_CHECK_API_MISMATCH(EOS_USERINFO_COPYBESTDISPLAYNAME_API_LATEST, 1);
 	Options.LocalUserId = GetEpicAccountIdChecked(Params.LocalAccountId);
 	Options.TargetUserId = TargetUserEasId;
 
-	EOS_UserInfo* EosUserInfo = nullptr;
+	EOS_UserInfo_BestDisplayName* EosBestDisplayName;
 	ON_SCOPE_EXIT
 	{
-		EOS_UserInfo_Release(EosUserInfo);
+		EOS_UserInfo_BestDisplayName_Release(EosBestDisplayName);
 	};
 
-	const EOS_EResult EosResult = EOS_UserInfo_CopyUserInfo(UserInfoHandle, &Options, &EosUserInfo);
+	EOS_EResult EosResult = EOS_UserInfo_CopyBestDisplayName(UserInfoHandle, &Options, &EosBestDisplayName);
+
+	if (EosResult == EOS_EResult::EOS_UserInfo_BestDisplayNameIndeterminate)
+	{
+		EOS_UserInfo_CopyBestDisplayNameWithPlatformOptions WithPlatformOptions = {};
+		WithPlatformOptions.ApiVersion = 1;
+		UE_EOS_CHECK_API_MISMATCH(EOS_USERINFO_COPYBESTDISPLAYNAMEWITHPLATFORM_API_LATEST, 1);
+		WithPlatformOptions.LocalUserId = GetEpicAccountIdChecked(Params.LocalAccountId);
+		WithPlatformOptions.TargetUserId = TargetUserEasId;
+		WithPlatformOptions.TargetPlatformType = EOS_OPT_Epic;
+
+		EosResult = EOS_UserInfo_CopyBestDisplayNameWithPlatform(UserInfoHandle, &WithPlatformOptions, &EosBestDisplayName);
+	}
+
 	if(EosResult != EOS_EResult::EOS_Success)
 	{
-		UE_LOG(LogOnlineServices, Warning, TEXT("EOS_UserInfo_CopyUserInfo failed with result=[%s]"), *LexToString(EosResult));
+		UE_LOG(LogOnlineServices, Warning, TEXT("EOS_UserInfo_CopyBestDisplayName failed with result=[%s]"), *LexToString(EosResult));
 		return TOnlineResult<FGetUserInfo>(Errors::FromEOSResult(EosResult));
 	}
 
 	TSharedRef<FUserInfo> UserInfo = MakeShared<FUserInfo>();
 	UserInfo->AccountId = Params.AccountId;
-	UserInfo->DisplayName = UTF8_TO_TCHAR(EosUserInfo->DisplayName);
+	UserInfo->DisplayName = UTF8_TO_TCHAR(EosBestDisplayName->DisplayNameSanitized);
 
 	return TOnlineResult<FGetUserInfo>({UserInfo});
 }
diff --git a/Engine/Plugins/Online/OnlineSubsystemEOS/Source/OnlineSubsystemEOS/Private/UserManagerEOS.cpp b/Engine/Plugins/Online/OnlineSubsystemEOS/Source/OnlineSubsystemEOS/Private/UserManagerEOS.cpp
--- a/Engine/Plugins/Online/OnlineSubsystemEOS/Source/OnlineSubsystemEOS/Private/UserManagerEOS.cpp
+++ b/Engine/Plugins/Online/OnlineSubsystemEOS/Source/OnlineSubsystemEOS/Private/UserManagerEOS.cpp
@@ -1075,11 +1075,42 @@ void FUserManagerEOS::UpdateUserInfo(IAttributeAccessInterfaceRef AttributeAcces
 	EOS_EResult CopyResult = EOS_UserInfo_CopyUserInfo(EOSSubsystem->UserInfoHandle, &Options, &UserInfo);
 	if (CopyResult == EOS_EResult::EOS_Success)
 	{
-		AttributeAccessRef->SetInternalAttribute(USER_ATTR_DISPLAY_NAME, UTF8_TO_TCHAR(UserInfo->DisplayName));
 		AttributeAccessRef->SetInternalAttribute(USER_ATTR_COUNTRY, UTF8_TO_TCHAR(UserInfo->Country));
 		AttributeAccessRef->SetInternalAttribute(USER_ATTR_LANG, UTF8_TO_TCHAR(UserInfo->PreferredLanguage));
 		EOS_UserInfo_Release(UserInfo);
 	}
+
+	EOS_UserInfo_CopyBestDisplayNameOptions BestDisplayNameOptions = { };
+	Options.ApiVersion = 1;
+	UE_EOS_CHECK_API_MISMATCH(EOS_USERINFO_COPYBESTDISPLAYNAME_API_LATEST, 1);
+	Options.LocalUserId = LocalId;
+	Options.TargetUserId = AccountId;
+
+	EOS_UserInfo_BestDisplayName* BestDisplayName;
+	EOS_EResult BestDisplayNameResult = EOS_UserInfo_CopyBestDisplayName(EOSSubsystem->UserInfoHandle, &BestDisplayNameOptions, &BestDisplayName);
+
+	if (BestDisplayNameResult == EOS_EResult::EOS_UserInfo_BestDisplayNameIndeterminate)
+	{
+		EOS_UserInfo_CopyBestDisplayNameWithPlatformOptions BestDisplayNameWithPlatformOptions = {};
+		BestDisplayNameWithPlatformOptions.ApiVersion = 1;
+		UE_EOS_CHECK_API_MISMATCH(EOS_USERINFO_COPYBESTDISPLAYNAMEWITHPLATFORM_API_LATEST, 1);
+		BestDisplayNameWithPlatformOptions.LocalUserId = LocalId;
+		BestDisplayNameWithPlatformOptions.TargetUserId = AccountId;
+		BestDisplayNameWithPlatformOptions.TargetPlatformType = EOS_OPT_Epic;
+
+		BestDisplayNameResult = EOS_UserInfo_CopyBestDisplayNameWithPlatform(EOSSubsystem->UserInfoHandle, &BestDisplayNameWithPlatformOptions, &BestDisplayName);
+	}
+
+	if (BestDisplayNameResult == EOS_EResult::EOS_Success)
+	{
+		AttributeAccessRef->SetInternalAttribute(USER_ATTR_DISPLAY_NAME, UTF8_TO_TCHAR(BestDisplayName->DisplayNameSanitized));
+
+		EOS_UserInfo_BestDisplayName_Release(BestDisplayName);
+	}
+	else
+	{
+		UE_LOG_ONLINE(Warning, TEXT("[FUserManagerEOS::UpdateUserInfo] Failed to retrieve BestDisplayName EOS_EResult: %s."), *LexToString(BestDisplayNameResult));
+	}
 }
 
 TSharedPtr<FUserOnlineAccount> FUserManagerEOS::GetUserAccount(const FUniqueNetId& UserId) const
