- Add IEOSPlatformHandle::GetConfigName
- Modify OnlineServicesEOS*::GetEOSPlatformHandle to return IEOSPlatformHandlePtr to expose the ConfigName.

diff --git a/Engine/Plugins/Online/EOSShared/Source/EOSShared/Private/EOSSDKManager.cpp b/Engine/Plugins/Online/EOSShared/Source/EOSShared/Private/EOSSDKManager.cpp
--- a/Engine/Plugins/Online/EOSShared/Source/EOSShared/Private/EOSSDKManager.cpp
+++ b/Engine/Plugins/Online/EOSShared/Source/EOSShared/Private/EOSSDKManager.cpp
@@ -511,7 +511,12 @@ IEOSPlatformHandlePtr FEOSSDKManager::CreatePlatform(const FString& PlatformConf
 IEOSPlatformHandlePtr FEOSSDKManager::CreatePlatform(const FEOSSDKPlatformConfig& PlatformConfig, EOS_Platform_Options& PlatformOptions)
 {
 	OnPreCreateNamedPlatform.Broadcast(PlatformConfig, PlatformOptions);
-	return CreatePlatform(PlatformOptions);
+	IEOSPlatformHandlePtr Result = CreatePlatform(PlatformOptions);
+	if (Result)
+	{
+		static_cast<FEOSPlatformHandle&>(*Result.Get()).ConfigName = PlatformConfig.Name;
+	}
+	return Result;
 }
 
 IEOSPlatformHandlePtr FEOSSDKManager::CreatePlatform(EOS_Platform_Options& PlatformOptions)
@@ -1168,6 +1173,11 @@ void FEOSPlatformHandle::Tick()
 	EOS_Platform_Tick(PlatformHandle);
 }
 
+FString FEOSPlatformHandle::GetConfigName() const
+{
+	return ConfigName;
+}
+
 FString FEOSPlatformHandle::GetOverrideCountryCode() const
 {
 	return Manager.GetOverrideCountryCode(PlatformHandle);
diff --git a/Engine/Plugins/Online/EOSShared/Source/EOSShared/Private/EOSSDKManager.h b/Engine/Plugins/Online/EOSShared/Source/EOSShared/Private/EOSSDKManager.h
--- a/Engine/Plugins/Online/EOSShared/Source/EOSShared/Private/EOSSDKManager.h
+++ b/Engine/Plugins/Online/EOSShared/Source/EOSShared/Private/EOSSDKManager.h
@@ -122,6 +122,7 @@ struct FEOSPlatformHandle : public IEOSPlatformHandle
 	virtual ~FEOSPlatformHandle();
 	virtual void Tick() override;
 
+	virtual FString GetConfigName() const override;
 	virtual FString GetOverrideCountryCode() const override;
 	virtual FString GetOverrideLocaleCode() const override;
 
@@ -134,6 +135,9 @@ struct FEOSPlatformHandle : public IEOSPlatformHandle
 
 	/* Reference to the EOSSDK manager */
 	FEOSSDKManager& Manager;
+
+	/* The name of the config used to instantiate this handle */
+	FString ConfigName;
 };
 
 #endif // WITH_EOS_SDK
\ No newline at end of file
diff --git a/Engine/Plugins/Online/EOSShared/Source/EOSShared/Public/IEOSSDKManager.h b/Engine/Plugins/Online/EOSShared/Source/EOSShared/Public/IEOSSDKManager.h
--- a/Engine/Plugins/Online/EOSShared/Source/EOSShared/Public/IEOSSDKManager.h
+++ b/Engine/Plugins/Online/EOSShared/Source/EOSShared/Public/IEOSSDKManager.h
@@ -45,13 +45,17 @@ DECLARE_MULTICAST_DELEGATE_OneParam(FEOSSDKManagerOnPreCreatePlatform, EOS_Platf
 class IEOSPlatformHandle
 {
 public:
-	IEOSPlatformHandle(EOS_HPlatform InPlatformHandle) : PlatformHandle(InPlatformHandle) {}
+	IEOSPlatformHandle(EOS_HPlatform InPlatformHandle)
+		: PlatformHandle(InPlatformHandle)
+	{}
+
 	virtual ~IEOSPlatformHandle() = default;
 
 	virtual void Tick() = 0;
 
 	operator EOS_HPlatform() const { return PlatformHandle; }
 
+	virtual FString GetConfigName() const = 0;
 	virtual FString GetOverrideCountryCode() const = 0;
 	virtual FString GetOverrideLocaleCode() const = 0;
 
diff --git a/Engine/Plugins/Online/OnlineServicesEOS/Source/Private/Online/AuthEOS.cpp b/Engine/Plugins/Online/OnlineServicesEOS/Source/Private/Online/AuthEOS.cpp
--- a/Engine/Plugins/Online/OnlineServicesEOS/Source/Private/Online/AuthEOS.cpp
+++ b/Engine/Plugins/Online/OnlineServicesEOS/Source/Private/Online/AuthEOS.cpp
@@ -2,14 +2,15 @@
 
 #include "Online/AuthEOS.h"
 
+#include "Algo/Transform.h"
 #include "EOSShared.h"
-#include "Online/OnlineIdEOS.h"
+#include "IEOSSDKManager.h"
+#include "Online/AuthErrors.h"
 #include "Online/OnlineErrorEOSGS.h"
+#include "Online/OnlineIdEOS.h"
 #include "Online/OnlineServicesEOS.h"
-#include "Algo/Transform.h"
-#include "Online/AuthErrors.h"
-
 #include "Online/OnlineUtils.h"
+
 #include "eos_auth.h"
 #include "eos_connect.h"
 #include "eos_userinfo.h"
@@ -46,7 +47,7 @@ void FAuthEOS::Initialize()
 {
 	Super::Initialize();
 
-	UserInfoHandle = EOS_Platform_GetUserInfoInterface(static_cast<FOnlineServicesEOS&>(GetServices()).GetEOSPlatformHandle());
+	UserInfoHandle = EOS_Platform_GetUserInfoInterface(*static_cast<FOnlineServicesEOS&>(GetServices()).GetEOSPlatformHandle());
 	check(UserInfoHandle != nullptr);
 }
 
diff --git a/Engine/Plugins/Online/OnlineServicesEOS/Source/Private/Online/CommerceEOS.cpp b/Engine/Plugins/Online/OnlineServicesEOS/Source/Private/Online/CommerceEOS.cpp
--- a/Engine/Plugins/Online/OnlineServicesEOS/Source/Private/Online/CommerceEOS.cpp
+++ b/Engine/Plugins/Online/OnlineServicesEOS/Source/Private/Online/CommerceEOS.cpp
@@ -1,10 +1,12 @@
 // Copyright Epic Games, Inc. All Rights Reserved.
 
 #include "Online/CommerceEOS.h"
+
+#include "IEOSSDKManager.h"
+#include "Online/AuthEOS.h"
 #include "Online/OnlineErrorEOSGS.h"
-#include "Online/OnlineServicesEOS.h"
 #include "Online/OnlineIdEOS.h"
-#include "Online/AuthEOS.h"
+#include "Online/OnlineServicesEOS.h"
 
 namespace UE::Online {
 
@@ -12,7 +14,7 @@ void FCommerceEOS::PostInitialize()
 {	
 	Super::PostInitialize();
 
-	EcomHandle = EOS_Platform_GetEcomInterface(static_cast<FOnlineServicesEOS&>(GetServices()).GetEOSPlatformHandle());
+	EcomHandle = EOS_Platform_GetEcomInterface(*static_cast<FOnlineServicesEOS&>(GetServices()).GetEOSPlatformHandle());
 	check(EcomHandle);
 }
 
diff --git a/Engine/Plugins/Online/OnlineServicesEOS/Source/Private/Online/ExternalUIEOS.cpp b/Engine/Plugins/Online/OnlineServicesEOS/Source/Private/Online/ExternalUIEOS.cpp
--- a/Engine/Plugins/Online/OnlineServicesEOS/Source/Private/Online/ExternalUIEOS.cpp
+++ b/Engine/Plugins/Online/OnlineServicesEOS/Source/Private/Online/ExternalUIEOS.cpp
@@ -3,10 +3,11 @@
 #include "Online/ExternalUIEOS.h"
 
 #include "EOSShared.h"
+#include "IEOSSDKManager.h"
 #include "Online/OnlineIdEOS.h"
 #include "Online/OnlineServicesEOS.h"
-
 #include "Online/OnlineServicesEOSGSTypes.h"
+
 #include "eos_ui.h"
 
 namespace UE::Online {
@@ -20,7 +21,7 @@ void FExternalUIEOS::Initialize()
 {
 	Super::Initialize();
 
-	UIHandle = EOS_Platform_GetUIInterface(static_cast<FOnlineServicesEOS&>(GetServices()).GetEOSPlatformHandle());
+	UIHandle = EOS_Platform_GetUIInterface(*static_cast<FOnlineServicesEOS&>(GetServices()).GetEOSPlatformHandle());
 	check(UIHandle != nullptr);
 }
 
diff --git a/Engine/Plugins/Online/OnlineServicesEOS/Source/Private/Online/PresenceEOS.cpp b/Engine/Plugins/Online/OnlineServicesEOS/Source/Private/Online/PresenceEOS.cpp
--- a/Engine/Plugins/Online/OnlineServicesEOS/Source/Private/Online/PresenceEOS.cpp
+++ b/Engine/Plugins/Online/OnlineServicesEOS/Source/Private/Online/PresenceEOS.cpp
@@ -1,11 +1,13 @@
 // Copyright Epic Games, Inc. All Rights Reserved.
 
 #include "Online/PresenceEOS.h"
+
 #include "EOSShared.h"
-#include "Online/OnlineIdEOS.h"
-#include "Online/OnlineServicesEOS.h"
+#include "IEOSSDKManager.h"
 #include "Online/AuthEOS.h"
 #include "Online/OnlineErrorEOSGS.h"
+#include "Online/OnlineIdEOS.h"
+#include "Online/OnlineServicesEOS.h"
 
 #include "eos_presence.h"
 
@@ -68,7 +70,7 @@ void FPresenceEOS::Initialize()
 {
 	FPresenceCommon::Initialize();
 
-	PresenceHandle = EOS_Platform_GetPresenceInterface(static_cast<FOnlineServicesEOS&>(GetServices()).GetEOSPlatformHandle());
+	PresenceHandle = EOS_Platform_GetPresenceInterface(*static_cast<FOnlineServicesEOS&>(GetServices()).GetEOSPlatformHandle());
 	check(PresenceHandle != nullptr);
 
 	// Register for friend updates
diff --git a/Engine/Plugins/Online/OnlineServicesEOS/Source/Private/Online/SocialEOS.cpp b/Engine/Plugins/Online/OnlineServicesEOS/Source/Private/Online/SocialEOS.cpp
--- a/Engine/Plugins/Online/OnlineServicesEOS/Source/Private/Online/SocialEOS.cpp
+++ b/Engine/Plugins/Online/OnlineServicesEOS/Source/Private/Online/SocialEOS.cpp
@@ -1,10 +1,12 @@
 // Copyright Epic Games, Inc. All Rights Reserved.
 
 #include "Online/SocialEOS.h"
+
+#include "IEOSSDKManager.h"
+#include "Online/AuthEOS.h"
 #include "Online/OnlineErrorEOSGS.h"
-#include "Online/OnlineServicesEOS.h"
 #include "Online/OnlineIdEOS.h"
-#include "Online/AuthEOS.h"
+#include "Online/OnlineServicesEOS.h"
 
 #include "eos_friends.h"
 #include "eos_ui.h"
@@ -33,10 +35,10 @@ void FSocialEOS::Initialize()
 {
 	FSocialCommon::Initialize();
 
-	FriendsHandle = EOS_Platform_GetFriendsInterface(static_cast<FOnlineServicesEOS&>(GetServices()).GetEOSPlatformHandle());
+	FriendsHandle = EOS_Platform_GetFriendsInterface(*static_cast<FOnlineServicesEOS&>(GetServices()).GetEOSPlatformHandle());
 	check(FriendsHandle != nullptr);
 
-	UIHandle = EOS_Platform_GetUIInterface(static_cast<FOnlineServicesEOS&>(GetServices()).GetEOSPlatformHandle());
+	UIHandle = EOS_Platform_GetUIInterface(*static_cast<FOnlineServicesEOS&>(GetServices()).GetEOSPlatformHandle());
 
 	// Register for friend updates
 	EOS_Friends_AddNotifyFriendsUpdateOptions Options = { };
diff --git a/Engine/Plugins/Online/OnlineServicesEOS/Source/Private/Online/UserInfoEOS.cpp b/Engine/Plugins/Online/OnlineServicesEOS/Source/Private/Online/UserInfoEOS.cpp
--- a/Engine/Plugins/Online/OnlineServicesEOS/Source/Private/Online/UserInfoEOS.cpp
+++ b/Engine/Plugins/Online/OnlineServicesEOS/Source/Private/Online/UserInfoEOS.cpp
@@ -1,7 +1,9 @@
 // Copyright Epic Games, Inc. All Rights Reserved.
 
 #include "Online/UserInfoEOS.h"
+
 #include "EOSShared.h"
+#include "IEOSSDKManager.h"
 #include "Online/OnlineIdEOS.h"
 #include "Online/OnlineServicesEOS.h"
 #include "Online/AuthEOS.h"
@@ -20,7 +22,7 @@ void FUserInfoEOS::Initialize()
 {
 	Super::Initialize();
 
-	UserInfoHandle = EOS_Platform_GetUserInfoInterface(static_cast<FOnlineServicesEOS&>(GetServices()).GetEOSPlatformHandle());
+	UserInfoHandle = EOS_Platform_GetUserInfoInterface(*static_cast<FOnlineServicesEOS&>(GetServices()).GetEOSPlatformHandle());
 	check(UserInfoHandle != nullptr);
 }
 
diff --git a/Engine/Plugins/Online/OnlineServicesEOSGS/Source/Private/Online/AchievementsEOSGS.cpp b/Engine/Plugins/Online/OnlineServicesEOSGS/Source/Private/Online/AchievementsEOSGS.cpp
--- a/Engine/Plugins/Online/OnlineServicesEOSGS/Source/Private/Online/AchievementsEOSGS.cpp
+++ b/Engine/Plugins/Online/OnlineServicesEOSGS/Source/Private/Online/AchievementsEOSGS.cpp
@@ -3,6 +3,7 @@
 #include "Online/AchievementsEOSGS.h"
 
 #include "EOSShared.h"
+#include "IEOSSDKManager.h"
 #include "Online/AchievementsErrors.h"
 #include "Online/AuthEOSGS.h"
 #include "Online/OnlineErrorEOSGS.h"
@@ -22,7 +23,7 @@ void FAchievementsEOSGS::Initialize()
 {
 	Super::Initialize();
 
-	AchievementsHandle = EOS_Platform_GetAchievementsInterface(static_cast<FOnlineServicesEOSGS&>(GetServices()).GetEOSPlatformHandle());
+	AchievementsHandle = EOS_Platform_GetAchievementsInterface(*static_cast<FOnlineServicesEOSGS&>(GetServices()).GetEOSPlatformHandle());
 	check(AchievementsHandle);
 
 	// Register for achievement unlocked events
diff --git a/Engine/Plugins/Online/OnlineServicesEOSGS/Source/Private/Online/AuthEOSGS.cpp b/Engine/Plugins/Online/OnlineServicesEOSGS/Source/Private/Online/AuthEOSGS.cpp
--- a/Engine/Plugins/Online/OnlineServicesEOSGS/Source/Private/Online/AuthEOSGS.cpp
+++ b/Engine/Plugins/Online/OnlineServicesEOSGS/Source/Private/Online/AuthEOSGS.cpp
@@ -4,13 +4,14 @@
 
 #include "Containers/BackgroundableTicker.h"
 #include "EOSShared.h"
+#include "IEOSSDKManager.h"
 #include "Misc/CommandLine.h"
 #include "Online/AuthErrors.h"
 #include "Online/OnlineErrorEOSGS.h"
 #include "Online/OnlineIdEOSGS.h"
 #include "Online/OnlineServicesEOSGS.h"
-
 #include "Online/OnlineUtils.h"
+
 #include "eos_auth.h"
 #include "eos_connect.h"
 
@@ -546,10 +547,10 @@ void FAuthEOSGS::Initialize()
 {
 	Super::Initialize();
 
-	AuthHandle = EOS_Platform_GetAuthInterface(static_cast<FOnlineServicesEOSGS&>(GetServices()).GetEOSPlatformHandle());
+	AuthHandle = EOS_Platform_GetAuthInterface(*static_cast<FOnlineServicesEOSGS&>(GetServices()).GetEOSPlatformHandle());
 	check(AuthHandle != nullptr);
 
-	ConnectHandle = EOS_Platform_GetConnectInterface(static_cast<FOnlineServicesEOSGS&>(GetServices()).GetEOSPlatformHandle());
+	ConnectHandle = EOS_Platform_GetConnectInterface(*static_cast<FOnlineServicesEOSGS&>(GetServices()).GetEOSPlatformHandle());
 	check(ConnectHandle != nullptr);
 
 	RegisterHandlers();
diff --git a/Engine/Plugins/Online/OnlineServicesEOSGS/Source/Private/Online/LeaderboardsEOSGS.cpp b/Engine/Plugins/Online/OnlineServicesEOSGS/Source/Private/Online/LeaderboardsEOSGS.cpp
--- a/Engine/Plugins/Online/OnlineServicesEOSGS/Source/Private/Online/LeaderboardsEOSGS.cpp
+++ b/Engine/Plugins/Online/OnlineServicesEOSGS/Source/Private/Online/LeaderboardsEOSGS.cpp
@@ -1,7 +1,9 @@
 // Copyright Epic Games, Inc. All Rights Reserved.
 
 #include "Online/LeaderboardsEOSGS.h"
+
 #include "EOSShared.h"
+#include "IEOSSDKManager.h"
 #include "Online/AuthEOSGS.h"
 #include "Online/OnlineErrorEOSGS.h"
 #include "Online/OnlineIdEOSGS.h"
@@ -23,7 +25,7 @@ void FLeaderboardsEOSGS::Initialize()
 {
 	Super::Initialize();
 
-	LeaderboardsHandle = EOS_Platform_GetLeaderboardsInterface(static_cast<FOnlineServicesEOSGS&>(GetServices()).GetEOSPlatformHandle());
+	LeaderboardsHandle = EOS_Platform_GetLeaderboardsInterface(*static_cast<FOnlineServicesEOSGS&>(GetServices()).GetEOSPlatformHandle());
 	check(LeaderboardsHandle);
 }
 
diff --git a/Engine/Plugins/Online/OnlineServicesEOSGS/Source/Private/Online/LobbiesEOSGS.cpp b/Engine/Plugins/Online/OnlineServicesEOSGS/Source/Private/Online/LobbiesEOSGS.cpp
--- a/Engine/Plugins/Online/OnlineServicesEOSGS/Source/Private/Online/LobbiesEOSGS.cpp
+++ b/Engine/Plugins/Online/OnlineServicesEOSGS/Source/Private/Online/LobbiesEOSGS.cpp
@@ -35,7 +35,7 @@ void FLobbiesEOSGS::Initialize()
 	IEOSSDKManager* SDKManager = IEOSSDKManager::Get();
 	check(SDKManager);
 
-	EOS_HLobby LobbyInterfaceHandle = EOS_Platform_GetLobbyInterface(static_cast<FOnlineServicesEOSGS&>(GetServices()).GetEOSPlatformHandle());
+	EOS_HLobby LobbyInterfaceHandle = EOS_Platform_GetLobbyInterface(*static_cast<FOnlineServicesEOSGS&>(GetServices()).GetEOSPlatformHandle());
 	check(LobbyInterfaceHandle != nullptr);
 
 	LobbyPrerequisites = MakeShared<FLobbyPrerequisitesEOS>(FLobbyPrerequisitesEOS{
diff --git a/Engine/Plugins/Online/OnlineServicesEOSGS/Source/Private/Online/OnlineServicesEOSGS.cpp b/Engine/Plugins/Online/OnlineServicesEOSGS/Source/Private/Online/OnlineServicesEOSGS.cpp
--- a/Engine/Plugins/Online/OnlineServicesEOSGS/Source/Private/Online/OnlineServicesEOSGS.cpp
+++ b/Engine/Plugins/Online/OnlineServicesEOSGS/Source/Private/Online/OnlineServicesEOSGS.cpp
@@ -81,11 +81,11 @@ void FOnlineServicesEOSGS::RegisterComponents()
 	Components.Register<FLobbiesEOSGS>(*this);
 	Components.Register<FStatsEOSGS>(*this);
 	Components.Register<FSessionsEOSGS>(*this);
-	if (EOS_Platform_GetTitleStorageInterface(GetEOSPlatformHandle()))
+	if (EOS_Platform_GetTitleStorageInterface(*GetEOSPlatformHandle()))
 	{
 		Components.Register<FTitleFileEOSGS>(*this);
 	}
-	if (EOS_Platform_GetPlayerDataStorageInterface(GetEOSPlatformHandle()))
+	if (EOS_Platform_GetPlayerDataStorageInterface(*GetEOSPlatformHandle()))
 	{
 		Components.Register<FUserFileEOSGS>(*this);
 	}
@@ -157,14 +157,4 @@ TOnlineResult<FGetResolvedConnectString> FOnlineServicesEOSGS::GetResolvedConnec
 	return TOnlineResult<FGetResolvedConnectString>(Errors::InvalidParams());
 }
 
-EOS_HPlatform FOnlineServicesEOSGS::GetEOSPlatformHandle() const
-{
-	if (EOSPlatformHandle)
-	{
-		return *EOSPlatformHandle;
-	}
-	return nullptr;
-}
-
-
 /* UE::Online */ }
diff --git a/Engine/Plugins/Online/OnlineServicesEOSGS/Source/Private/Online/SessionsEOSGS.cpp b/Engine/Plugins/Online/OnlineServicesEOSGS/Source/Private/Online/SessionsEOSGS.cpp
--- a/Engine/Plugins/Online/OnlineServicesEOSGS/Source/Private/Online/SessionsEOSGS.cpp
+++ b/Engine/Plugins/Online/OnlineServicesEOSGS/Source/Private/Online/SessionsEOSGS.cpp
@@ -3,11 +3,12 @@
 #include "Online/SessionsEOSGS.h"
 
 #include "EOSShared.h"
+#include "IEOSSDKManager.h"
 #include "Online/AuthEOSGS.h"
 #include "Online/LobbiesCommonTypes.h"
-#include "Online/OnlineServicesEOSGS.h"
 #include "Online/NboSerializerEOSGSSvc.h"
 #include "Online/OnlineErrorEOSGS.h"
+#include "Online/OnlineServicesEOSGS.h"
 #include "Online/SessionsEOSGSTypes.h"
 
 #include "eos_sessions.h"
@@ -143,7 +144,7 @@ void FSessionsEOSGS::Initialize()
 {
 	Super::Initialize();
 
-	SessionsHandle = EOS_Platform_GetSessionsInterface(static_cast<FOnlineServicesEOSGS&>(GetServices()).GetEOSPlatformHandle());
+	SessionsHandle = EOS_Platform_GetSessionsInterface(*static_cast<FOnlineServicesEOSGS&>(GetServices()).GetEOSPlatformHandle());
 	check(SessionsHandle);
 
 	RegisterEventHandlers();
diff --git a/Engine/Plugins/Online/OnlineServicesEOSGS/Source/Private/Online/StatsEOSGS.cpp b/Engine/Plugins/Online/OnlineServicesEOSGS/Source/Private/Online/StatsEOSGS.cpp
--- a/Engine/Plugins/Online/OnlineServicesEOSGS/Source/Private/Online/StatsEOSGS.cpp
+++ b/Engine/Plugins/Online/OnlineServicesEOSGS/Source/Private/Online/StatsEOSGS.cpp
@@ -1,7 +1,9 @@
 // Copyright Epic Games, Inc. All Rights Reserved.
 
 #include "Online/StatsEOSGS.h"
+
 #include "EOSShared.h"
+#include "IEOSSDKManager.h"
 #include "Online/AuthEOSGS.h"
 #include "Online/OnlineErrorEOSGS.h"
 #include "Online/OnlineIdEOSGS.h"
@@ -28,7 +30,7 @@ void FStatsEOSGS::Initialize()
 {
 	Super::Initialize();
 
-	StatsHandle = EOS_Platform_GetStatsInterface(static_cast<FOnlineServicesEOSGS&>(GetServices()).GetEOSPlatformHandle());
+	StatsHandle = EOS_Platform_GetStatsInterface(*static_cast<FOnlineServicesEOSGS&>(GetServices()).GetEOSPlatformHandle());
 	check(StatsHandle);
 }
 
diff --git a/Engine/Plugins/Online/OnlineServicesEOSGS/Source/Private/Online/TitleFileEOSGS.cpp b/Engine/Plugins/Online/OnlineServicesEOSGS/Source/Private/Online/TitleFileEOSGS.cpp
--- a/Engine/Plugins/Online/OnlineServicesEOSGS/Source/Private/Online/TitleFileEOSGS.cpp
+++ b/Engine/Plugins/Online/OnlineServicesEOSGS/Source/Private/Online/TitleFileEOSGS.cpp
@@ -3,6 +3,7 @@
 #include "Online/TitleFileEOSGS.h"
 
 #include "EOSShared.h"
+#include "IEOSSDKManager.h"
 #include "Online/AuthEOSGS.h"
 #include "Online/OnlineErrorEOSGS.h"
 #include "Online/OnlineIdEOSGS.h"
@@ -22,7 +23,7 @@ void FTitleFileEOSGS::Initialize()
 {
 	Super::Initialize();
 
-	TitleStorageHandle = EOS_Platform_GetTitleStorageInterface(static_cast<FOnlineServicesEOSGS&>(GetServices()).GetEOSPlatformHandle());
+	TitleStorageHandle = EOS_Platform_GetTitleStorageInterface(*static_cast<FOnlineServicesEOSGS&>(GetServices()).GetEOSPlatformHandle());
 	check(TitleStorageHandle);
 }
 
diff --git a/Engine/Plugins/Online/OnlineServicesEOSGS/Source/Private/Online/UserFileEOSGS.cpp b/Engine/Plugins/Online/OnlineServicesEOSGS/Source/Private/Online/UserFileEOSGS.cpp
--- a/Engine/Plugins/Online/OnlineServicesEOSGS/Source/Private/Online/UserFileEOSGS.cpp
+++ b/Engine/Plugins/Online/OnlineServicesEOSGS/Source/Private/Online/UserFileEOSGS.cpp
@@ -3,6 +3,7 @@
 #include "Online/UserFileEOSGS.h"
 
 #include "EOSShared.h"
+#include "IEOSSDKManager.h"
 #include "Online/AuthEOSGS.h"
 #include "Online/OnlineErrorEOSGS.h"
 #include "Online/OnlineIdEOSGS.h"
@@ -24,7 +25,7 @@ void FUserFileEOSGS::Initialize()
 {
 	Super::Initialize();
 
-	PlayerDataStorageHandle = EOS_Platform_GetPlayerDataStorageInterface(static_cast<FOnlineServicesEOSGS&>(GetServices()).GetEOSPlatformHandle());
+	PlayerDataStorageHandle = EOS_Platform_GetPlayerDataStorageInterface(*static_cast<FOnlineServicesEOSGS&>(GetServices()).GetEOSPlatformHandle());
 	check(PlayerDataStorageHandle);
 }
 
diff --git a/Engine/Plugins/Online/OnlineServicesEOSGS/Source/Public/Online/OnlineServicesEOSGS.h b/Engine/Plugins/Online/OnlineServicesEOSGS/Source/Public/Online/OnlineServicesEOSGS.h
--- a/Engine/Plugins/Online/OnlineServicesEOSGS/Source/Public/Online/OnlineServicesEOSGS.h
+++ b/Engine/Plugins/Online/OnlineServicesEOSGS/Source/Public/Online/OnlineServicesEOSGS.h
@@ -31,7 +31,7 @@ public:
 	virtual TOnlineResult<FGetResolvedConnectString> GetResolvedConnectString(FGetResolvedConnectString::Params&& Params) override;
 	virtual EOnlineServices GetServicesProvider() const override { return EOnlineServices::Epic; }
 
-	EOS_HPlatform GetEOSPlatformHandle() const;
+	IEOSPlatformHandlePtr GetEOSPlatformHandle() const { return EOSPlatformHandle; }
 
 	static const TCHAR* GetConfigNameStatic() { return TEXT("EOS"); }
 protected:
