Implement primary and secondary internal auth methods. Primary is OpenId token, secondary is auth refresh token.

diff --git a/Engine/Plugins/Online/OnlineServicesEOS/Source/Private/Online/AuthEOS.cpp b/Engine/Plugins/Online/OnlineServicesEOS/Source/Private/Online/AuthEOS.cpp
--- a/Engine/Plugins/Online/OnlineServicesEOS/Source/Private/Online/AuthEOS.cpp
+++ b/Engine/Plugins/Online/OnlineServicesEOS/Source/Private/Online/AuthEOS.cpp
@@ -342,15 +342,50 @@ TOnlineAsyncOpHandle<FAuthQueryExternalAuthToken> FAuthEOS::QueryExternalAuthTok
 				return;
 			}
 
-			TDefaultErrorResult<FAuthGetExternalAuthTokenImpl> AuthTokenResult = GetExternalAuthTokenImpl(FAuthGetExternalAuthTokenImpl::Params{ AccountInfoEOS->EpicAccountId });
-			if (AuthTokenResult.IsError())
+			// The primary external auth method is an id token.
+			if (Params.Method == EExternalAuthTokenMethod::Primary)
 			{
-				UE_LOG(LogOnlineServices, Warning, TEXT("[FAuthEOS::QueryExternalAuthToken] Failure: GetExternalAuthTokenImpl %s"), *AuthTokenResult.GetErrorValue().GetLogString());
-				InAsyncOp.SetError(Errors::Unknown(MoveTemp(AuthTokenResult.GetErrorValue())));
-				return;
+				TDefaultErrorResult<FAuthGetExternalAuthTokenImpl> AuthTokenResult = GetExternalAuthTokenImpl(FAuthGetExternalAuthTokenImpl::Params{ AccountInfoEOS->EpicAccountId });
+				if (AuthTokenResult.IsError())
+				{
+					UE_LOG(LogOnlineServices, Warning, TEXT("[FAuthEOS::QueryExternalAuthToken] Failure: GetExternalAuthTokenImpl %s"), *AuthTokenResult.GetErrorValue().GetLogString());
+					InAsyncOp.SetError(Errors::Unknown(MoveTemp(AuthTokenResult.GetErrorValue())));
+					return;
+				}
+
+				InAsyncOp.SetResult(FAuthQueryExternalAuthToken::Result{ MoveTemp(AuthTokenResult.GetOkValue().Token) });
 			}
+			// The secondary external auth method is an EAS refresh token.
+			else if (Params.Method == EExternalAuthTokenMethod::Secondary)
+			{
+				EOS_Auth_CopyUserAuthTokenOptions CopyUserAuthTokenOptions = {};
+				CopyUserAuthTokenOptions.ApiVersion = 1;
+				UE_EOS_CHECK_API_MISMATCH(EOS_AUTH_COPYUSERAUTHTOKEN_API_LATEST, 1);
+
+				EOS_Auth_Token* AuthToken = nullptr;
+
+				EOS_EResult Result = EOS_Auth_CopyUserAuthToken(AuthHandle, &CopyUserAuthTokenOptions, AccountInfoEOS->EpicAccountId, &AuthToken);
+				if (Result == EOS_EResult::EOS_Success)
+				{
+					ON_SCOPE_EXIT
+					{
+						EOS_Auth_Token_Release(AuthToken);
+					};
 
-			InAsyncOp.SetResult(FAuthQueryExternalAuthToken::Result{ MoveTemp(AuthTokenResult.GetOkValue().Token) });
+					FExternalAuthToken ExternalAuthToken;
+					ExternalAuthToken.Type = ExternalLoginType::Epic;
+					ExternalAuthToken.Data = UTF8_TO_TCHAR(AuthToken->RefreshToken);
+					InAsyncOp.SetResult(FAuthQueryExternalAuthToken::Result{ MoveTemp(ExternalAuthToken) });
+				}
+				else
+				{
+					InAsyncOp.SetError(Errors::FromEOSResult(Result));
+				}
+			}
+			else
+			{
+				InAsyncOp.SetError(Errors::InvalidParams());
+			}
 		})
 		.Enqueue(GetSerialQueue());
 	}
