Reduce verbosity level of message from UpdateUserPresence operation.

diff --git a/Engine/Plugins/Online/OnlineServicesEOS/Source/Private/Online/PresenceEOS.cpp b/Engine/Plugins/Online/OnlineServicesEOS/Source/Private/Online/PresenceEOS.cpp
--- a/Engine/Plugins/Online/OnlineServicesEOS/Source/Private/Online/PresenceEOS.cpp
+++ b/Engine/Plugins/Online/OnlineServicesEOS/Source/Private/Online/PresenceEOS.cpp
@@ -692,7 +692,7 @@ void FPresenceEOS::UpdateUserPresence(FAccountId LocalAccountId, FAccountId Pres
 	}
 	else
 	{
-		UE_LOG(LogOnlineServices, Error, TEXT("UpdateUserPresence: CopyPresence Failed %s"), *LexToString(CopyPresenceResult));
+		UE_LOG(LogOnlineServices, Warning, TEXT("UpdateUserPresence: CopyPresence Failed %s"), *LexToString(CopyPresenceResult));
 	}
 
 	if (bPresenceHasChanged)
