Integrate from 29.40 CL 33006298

Add IAuth::LinkAccount method to allow continuing a login which failed with "invalid user" by linking to an existing or new account.
Add missing parameter to IAuth::GetRelyingParty.
Add handling of EOS_ELinkAccountFlags.
Implement AuthEOS::LinkAccount

diff --git a/Engine/Plugins/Online/OnlineServices/Source/OnlineServicesCommon/Private/Online/AuthCommon.cpp b/Engine/Plugins/Online/OnlineServices/Source/OnlineServicesCommon/Private/Online/AuthCommon.cpp
--- a/Engine/Plugins/Online/OnlineServices/Source/OnlineServicesCommon/Private/Online/AuthCommon.cpp
+++ b/Engine/Plugins/Online/OnlineServices/Source/OnlineServicesCommon/Private/Online/AuthCommon.cpp
@@ -85,6 +85,7 @@ void FAuthCommon::RegisterCommands()
 {
 	RegisterCommand(&FAuthCommon::Login);
 	RegisterCommand(&FAuthCommon::Logout);
+	RegisterCommand(&FAuthCommon::LinkAccount);
 	RegisterCommand(&FAuthCommon::ModifyAccountAttributes);
 	RegisterCommand(&FAuthCommon::QueryExternalServerAuthTicket);
 	RegisterCommand(&FAuthCommon::QueryExternalAuthToken);
@@ -95,6 +96,8 @@ void FAuthCommon::RegisterCommands()
 	RegisterCommand(&FAuthCommon::GetLocalOnlineUserByOnlineAccountId);
 	RegisterCommand(&FAuthCommon::GetLocalOnlineUserByPlatformUserId);
 	RegisterCommand(&FAuthCommon::GetAllLocalOnlineUsers);
+	RegisterCommand(&FAuthCommon::GetLinkAccountContinuationId);
+	RegisterCommand(&FAuthCommon::GetRelyingParty);
 }
 
 TOnlineAsyncOpHandle<FAuthLogin> FAuthCommon::Login(FAuthLogin::Params&& Params)
@@ -111,6 +114,13 @@ TOnlineAsyncOpHandle<FAuthLogout> FAuthCommon::Logout(FAuthLogout::Params&& Para
 	return Operation->GetHandle();
 }
 
+TOnlineAsyncOpHandle<FAuthLinkAccount> FAuthCommon::LinkAccount(FAuthLinkAccount::Params&& Params)
+{
+	TOnlineAsyncOpRef<FAuthLinkAccount> Operation = GetOp<FAuthLinkAccount>(MoveTemp(Params));
+	Operation->SetError(Errors::NotImplemented());
+	return Operation->GetHandle();
+}
+
 TOnlineAsyncOpHandle<FAuthModifyAccountAttributes> FAuthCommon::ModifyAccountAttributes(FAuthModifyAccountAttributes::Params&& Params)
 {
 	TOnlineAsyncOpRef<FAuthModifyAccountAttributes> Operation = GetOp<FAuthModifyAccountAttributes>(MoveTemp(Params));
@@ -182,7 +192,12 @@ TOnlineResult<FAuthGetAllLocalOnlineUsers> FAuthCommon::GetAllLocalOnlineUsers(F
 		GetAccountInfoRegistry().GetAllAccountInfo([](const TSharedRef<FAccountInfo>& AccountInfo) { return IsOnlineStatus(AccountInfo->LoginStatus); }) });
 }
 
-TOnlineResult<FAuthGetRelyingParty> FAuthCommon::GetRelyingParty() const
+TOnlineResult<FAuthGetLinkAccountContinuationId> FAuthCommon::GetLinkAccountContinuationId(FAuthGetLinkAccountContinuationId::Params&& Params) const
+{
+	return TOnlineResult<FAuthGetLinkAccountContinuationId>(Errors::NotImplemented());
+}
+
+TOnlineResult<FAuthGetRelyingParty> FAuthCommon::GetRelyingParty(FAuthGetRelyingParty::Params&& Params) const
 {
 	return TOnlineResult<FAuthGetRelyingParty>(Errors::NotImplemented());
 }
diff --git a/Engine/Plugins/Online/OnlineServices/Source/OnlineServicesCommon/Public/Online/AuthCommon.h b/Engine/Plugins/Online/OnlineServices/Source/OnlineServicesCommon/Public/Online/AuthCommon.h
--- a/Engine/Plugins/Online/OnlineServices/Source/OnlineServicesCommon/Public/Online/AuthCommon.h
+++ b/Engine/Plugins/Online/OnlineServices/Source/OnlineServicesCommon/Public/Online/AuthCommon.h
@@ -45,6 +45,7 @@ public:
 	// IAuth
 	virtual TOnlineAsyncOpHandle<FAuthLogin> Login(FAuthLogin::Params&& Params) override;
 	virtual TOnlineAsyncOpHandle<FAuthLogout> Logout(FAuthLogout::Params&& Params) override;
+	virtual TOnlineAsyncOpHandle<FAuthLinkAccount> LinkAccount(FAuthLinkAccount::Params&& Params) override;
 	virtual TOnlineAsyncOpHandle<FAuthModifyAccountAttributes> ModifyAccountAttributes(FAuthModifyAccountAttributes::Params&& Params) override;
 	virtual TOnlineAsyncOpHandle<FAuthQueryExternalServerAuthTicket> QueryExternalServerAuthTicket(FAuthQueryExternalServerAuthTicket::Params&& Params) override;
 	virtual TOnlineAsyncOpHandle<FAuthQueryExternalAuthToken> QueryExternalAuthToken(FAuthQueryExternalAuthToken::Params&& Params) override;
@@ -55,7 +56,8 @@ public:
 	virtual TOnlineResult<FAuthGetLocalOnlineUserByOnlineAccountId> GetLocalOnlineUserByOnlineAccountId(FAuthGetLocalOnlineUserByOnlineAccountId::Params&& Params) const override;
 	virtual TOnlineResult<FAuthGetLocalOnlineUserByPlatformUserId> GetLocalOnlineUserByPlatformUserId(FAuthGetLocalOnlineUserByPlatformUserId::Params&& Params) const override;
 	virtual TOnlineResult<FAuthGetAllLocalOnlineUsers> GetAllLocalOnlineUsers(FAuthGetAllLocalOnlineUsers::Params&& Params) const override;
-	virtual TOnlineResult<FAuthGetRelyingParty> GetRelyingParty() const override;
+	virtual TOnlineResult<FAuthGetLinkAccountContinuationId> GetLinkAccountContinuationId(FAuthGetLinkAccountContinuationId::Params&& Params) const override;
+	virtual TOnlineResult<FAuthGetRelyingParty> GetRelyingParty(FAuthGetRelyingParty::Params&& Params) const override;
 	virtual TOnlineEvent<void(const FAuthLoginStatusChanged&)> OnLoginStatusChanged() override;
 	virtual TOnlineEvent<void(const FAuthPendingAuthExpiration&)> OnPendingAuthExpiration() override;
 	virtual TOnlineEvent<void(const FAuthAccountAttributesChanged&)> OnAccountAttributesChanged() override;
diff --git a/Engine/Plugins/Online/OnlineServices/Source/OnlineServicesInterface/Public/Online/Auth.h b/Engine/Plugins/Online/OnlineServices/Source/OnlineServicesInterface/Public/Online/Auth.h
--- a/Engine/Plugins/Online/OnlineServices/Source/OnlineServicesInterface/Public/Online/Auth.h
+++ b/Engine/Plugins/Online/OnlineServices/Source/OnlineServicesInterface/Public/Online/Auth.h
@@ -189,6 +189,24 @@ struct FAuthLogout
 	};
 };
 
+struct FAuthLinkAccount
+{
+	static constexpr TCHAR Name[] = TEXT("LinkAccount");
+
+	struct Params
+	{
+		/** The PlatformUserId of the Local User making the request. */
+		FPlatformUserId PlatformUserId = PLATFORMUSERID_NONE;
+		/** The id associated with external auth data used to continue a previous login attempt. */
+		FLoginContinuationId ContinuationId;
+	};
+
+	struct Result
+	{
+		TSharedRef<const FAccountInfo> AccountInfo;
+	};
+};
+
 struct FAuthModifyAccountAttributes
 {
 	static constexpr TCHAR Name[] = TEXT("ModifyAccountAttributes");
@@ -360,6 +378,23 @@ struct FAuthGetAllLocalOnlineUsers
 	};
 };
 
+struct FAuthGetLinkAccountContinuationId
+{
+	static constexpr TCHAR Name[] = TEXT("GetLinkAccountContinuationId");
+
+	struct Params
+	{
+		/** The PlatformUserId of the Local User making the request. */
+		FPlatformUserId PlatformUserId = PLATFORMUSERID_NONE;
+	};
+
+	struct Result
+	{
+		/** The continuation id which identifies the external auth to be linked. */
+		FLoginContinuationId ContinuationId;
+	};
+};
+
 struct FAuthGetRelyingParty
 {
 	static constexpr TCHAR Name[] = TEXT("GetRelyingParty");
@@ -409,6 +444,10 @@ public:
 	/**
 	 * Authenticate a local user. If necessary, retrieve and maintain an access token for the
 	 * duration of the auth session.
+	 * 
+	 * Returns Errors::InvalidUser when no user has been associated with the provided credentials.
+	 * GetLinkAccountContinuationId can then be used to retrieve a continuance token which can be
+	 * used to call LinkAccount to associate credentials with the service.
 	 */
 	virtual TOnlineAsyncOpHandle<FAuthLogin> Login(FAuthLogin::Params&& Params) = 0;
 
@@ -417,6 +456,11 @@ public:
 	 */
 	virtual TOnlineAsyncOpHandle<FAuthLogout> Logout(FAuthLogout::Params&& Params) = 0;
 
+	/**
+	 * Link a logged in account to an external auth method using a continuation token from a previous login attempt.
+	 */
+	virtual TOnlineAsyncOpHandle<FAuthLinkAccount> LinkAccount(FAuthLinkAccount::Params&& Params) = 0;
+
 	/**
 	 * Modify attributes associated with an authenticated account. Signals OnAccountAttributesChanged on completion.
 	 */
@@ -494,10 +538,15 @@ public:
 	 */
 	virtual TOnlineResult<FAuthGetAllLocalOnlineUsers> GetAllLocalOnlineUsers(FAuthGetAllLocalOnlineUsers::Params&& Params) const = 0;
 
+	/**
+	 * Retrieve the continuation id from a previous login attempt which used an external auth token and failed with Errors::InvalidUser.
+	 */
+	virtual TOnlineResult<FAuthGetLinkAccountContinuationId> GetLinkAccountContinuationId(FAuthGetLinkAccountContinuationId::Params&& Params) const = 0;
+
 	/**
 	 * Retrieve the relying party to use when fetching an external token to authenticate with this interface.
 	 */
-	virtual TOnlineResult<FAuthGetRelyingParty> GetRelyingParty() const = 0;
+	virtual TOnlineResult<FAuthGetRelyingParty> GetRelyingParty(FAuthGetRelyingParty::Params&& Params) const = 0;
 
 	/**
 	 * Triggered when the login status for a logged in user changes.
@@ -577,6 +626,15 @@ END_ONLINE_STRUCT_META()
 BEGIN_ONLINE_STRUCT_META(FAuthLogout::Result)
 END_ONLINE_STRUCT_META()
 
+BEGIN_ONLINE_STRUCT_META(FAuthLinkAccount::Params)
+	ONLINE_STRUCT_FIELD(FAuthLinkAccount::Params, PlatformUserId),
+	ONLINE_STRUCT_FIELD(FAuthLinkAccount::Params, ContinuationId)
+END_ONLINE_STRUCT_META()
+
+BEGIN_ONLINE_STRUCT_META(FAuthLinkAccount::Result)
+	ONLINE_STRUCT_FIELD(FAuthLinkAccount::Result, AccountInfo)
+END_ONLINE_STRUCT_META()
+
 BEGIN_ONLINE_STRUCT_META(FAuthModifyAccountAttributes::Params)
 	ONLINE_STRUCT_FIELD(FAuthModifyAccountAttributes::Params, LocalAccountId),
 	ONLINE_STRUCT_FIELD(FAuthModifyAccountAttributes::Params, MutatedAttributes),
@@ -659,6 +717,14 @@ BEGIN_ONLINE_STRUCT_META(FAuthGetAllLocalOnlineUsers::Result)
 	ONLINE_STRUCT_FIELD(FAuthGetAllLocalOnlineUsers::Result, AccountInfo)
 END_ONLINE_STRUCT_META()
 
+BEGIN_ONLINE_STRUCT_META(FAuthGetLinkAccountContinuationId::Params)
+	ONLINE_STRUCT_FIELD(FAuthGetLinkAccountContinuationId::Params, PlatformUserId)
+END_ONLINE_STRUCT_META()
+
+BEGIN_ONLINE_STRUCT_META(FAuthGetLinkAccountContinuationId::Result)
+	ONLINE_STRUCT_FIELD(FAuthGetLinkAccountContinuationId::Result, ContinuationId)
+END_ONLINE_STRUCT_META()
+
 BEGIN_ONLINE_STRUCT_META(FAuthGetRelyingParty::Params)
 END_ONLINE_STRUCT_META()
 
diff --git a/Engine/Plugins/Online/OnlineServicesEOS/Source/Private/Online/AuthEOS.cpp b/Engine/Plugins/Online/OnlineServicesEOS/Source/Private/Online/AuthEOS.cpp
--- a/Engine/Plugins/Online/OnlineServicesEOS/Source/Private/Online/AuthEOS.cpp
+++ b/Engine/Plugins/Online/OnlineServicesEOS/Source/Private/Online/AuthEOS.cpp
@@ -2,6 +2,8 @@
 
 #include "Online/AuthEOS.h"
 
+#include "Algo/Find.h"
+#include "Algo/RemoveIf.h"
 #include "Algo/Transform.h"
 #include "Containers/StaticArray.h"
 #include "EOSShared.h"
@@ -21,12 +23,14 @@ namespace UE::Online {
 struct FAuthEOSLoginConfig
 {
 	TArray<FString> DefaultScopes;
+	bool bAutoLinkAccount = true;
 };
 
 namespace Meta {
 
 BEGIN_ONLINE_STRUCT_META(FAuthEOSLoginConfig)
-	ONLINE_STRUCT_FIELD(FAuthEOSLoginConfig, DefaultScopes)
+	ONLINE_STRUCT_FIELD(FAuthEOSLoginConfig, DefaultScopes),
+	ONLINE_STRUCT_FIELD(FAuthEOSLoginConfig, bAutoLinkAccount)
 END_ONLINE_STRUCT_META()
 
 /* Meta*/ }
@@ -36,6 +40,8 @@ namespace
 {
 
 #define UE_ONLINE_AUTH_EOS_ACCOUNT_INFO_KEY_NAME TEXT("AccountInfoEOS")
+#define UE_ONLINE_AUTH_EOS_CONTINUANCE_DATA_KEY_NAME TEXT("ContinuanceToken")
+#define UE_ONLINE_AUTH_EOS_SELECTED_ACCOUNT_ID_KEY_NAME TEXT("SelectedAccountId")
 
 /* anonymous */ }
 
@@ -59,6 +65,14 @@ TOnlineAsyncOpHandle<FAuthLogin> FAuthEOS::Login(FAuthLogin::Params&& Params)
 	Op->Then([this](TOnlineAsyncOp<FAuthLogin>& InAsyncOp)
 	{
 		const FAuthLogin::Params& Params = InAsyncOp.GetParams();
+
+		// Check that user is valid.
+		if (!Params.PlatformUserId.IsValid())
+		{
+			InAsyncOp.SetError(Errors::InvalidParams());
+			return;
+		}
+
 		TSharedPtr<FAccountInfoEOS> AccountInfoEOS = AccountInfoRegistryEOS.Find(Params.PlatformUserId);
 		if (AccountInfoEOS)
 		{
@@ -70,6 +84,12 @@ TOnlineAsyncOpHandle<FAuthLogin> FAuthEOS::Login(FAuthLogin::Params&& Params)
 		AccountInfoEOS->PlatformUserId = Params.PlatformUserId;
 		AccountInfoEOS->LoginStatus = ELoginStatus::NotLoggedIn;
 
+		// New login attempt - Clear the continuance token for the last login attempt for the user.
+		if (FUserScopedData* UserData = GetUserScopedData(Params.PlatformUserId))
+		{
+			UserData->LastLoginContinuationId = FLoginContinuationId();
+		}
+
 		// Set user auth data on operation.
 		InAsyncOp.Data.Set<TSharedRef<FAccountInfoEOS>>(UE_ONLINE_AUTH_EOS_ACCOUNT_INFO_KEY_NAME, AccountInfoEOS.ToSharedRef());
 	})
@@ -90,6 +110,7 @@ TOnlineAsyncOpHandle<FAuthLogin> FAuthEOS::Login(FAuthLogin::Params&& Params)
 		LoginParams.CredentialsId = Params.CredentialsId;
 		LoginParams.CredentialsToken = Params.CredentialsToken;
 		LoginParams.Scopes = !Params.Scopes.IsEmpty() ? Params.Scopes : AuthEOSLoginConfig.DefaultScopes;
+		LoginParams.bAutoLinkAccount = AuthEOSLoginConfig.bAutoLinkAccount;
 
 		LoginEASImpl(LoginParams)
 		.Next([this, Promise = MoveTemp(Promise), WeakOp = InAsyncOp.AsWeak()](TDefaultErrorResult<FAuthLoginEASImpl>&& LoginResult) mutable -> void
@@ -128,18 +149,14 @@ TOnlineAsyncOpHandle<FAuthLogin> FAuthEOS::Login(FAuthLogin::Params&& Params)
 		if (AuthTokenResult.IsError())
 		{
 			UE_LOG(LogOnlineServices, Warning, TEXT("[FAuthEOS::Login] Failure: GetExternalAuthTokenImpl %s"), *AuthTokenResult.GetErrorValue().GetLogString());
-			InAsyncOp.SetError(Errors::Unknown(MoveTemp(AuthTokenResult.GetErrorValue())));
 
 			// Failed to acquire token - logout EAS.
 			LogoutEASImpl(FAuthLogoutEASImpl::Params{ AccountInfoEOS->EpicAccountId })
-			.Next([Promise = MoveTemp(Promise)](TDefaultErrorResult<FAuthLogoutEASImpl>&& LogoutResult) mutable -> void
-			{
-				if (LogoutResult.IsError())
+				.Next([AsyncOp = InAsyncOp.AsShared(), Error = MoveTemp(AuthTokenResult.GetErrorValue()), Promise = MoveTemp(Promise)](TDefaultErrorResult<FAuthLogoutEASImpl>&&) mutable -> void
 				{
-					UE_LOG(LogOnlineServices, Warning, TEXT("[FAuthEOS::Login] Failure: LogoutEASImpl %s"), *LogoutResult.GetErrorValue().GetLogString());
-				}
-				Promise.EmplaceValue(FAuthLoginConnectImpl::Params{});
-			});
+					AsyncOp->SetError(MoveTemp(Error));
+					Promise.EmplaceValue(FAuthLoginConnectImpl::Params{});
+				});
 
 			return Future;
 		}
@@ -159,20 +176,16 @@ TOnlineAsyncOpHandle<FAuthLogin> FAuthEOS::Login(FAuthLogin::Params&& Params)
 		LoginConnectImpl(LoginConnectParams)
 		.Next([this, AccountInfoEOS, WeakOp = InAsyncOp.AsWeak(), Promise = MoveTemp(Promise)](TDefaultErrorResult<FAuthLoginConnectImpl>&& LoginResult) mutable -> void
 		{
-			if (TSharedPtr<TOnlineAsyncOp<FAuthLogin>> Op = WeakOp.Pin())
+			if (TSharedPtr<TOnlineAsyncOp<FAuthLogin>> AsyncOp = WeakOp.Pin())
 			{
 				if (LoginResult.IsError())
 				{
 					UE_LOG(LogOnlineServices, Warning, TEXT("[FAuthEOS::Login] Failure: LoginConnectImpl %s"), *LoginResult.GetErrorValue().GetLogString());
-					Op->SetError(Errors::Unknown(MoveTemp(LoginResult.GetErrorValue())));
 
 					LogoutEASImpl(FAuthLogoutEASImpl::Params{ AccountInfoEOS->EpicAccountId })
-					.Next([Promise = MoveTemp(Promise)](TDefaultErrorResult<FAuthLogoutEASImpl>&& LogoutResult) mutable -> void
+					.Next([AsyncOp, Error = MoveTemp(LoginResult.GetErrorValue()), Promise = MoveTemp(Promise)](TDefaultErrorResult<FAuthLogoutEASImpl>&&) mutable -> void
 					{
-						if (LogoutResult.IsError())
-						{
-							UE_LOG(LogOnlineServices, Warning, TEXT("[FAuthEOS::Login] Failure: LogoutEASImpl %s"), *LogoutResult.GetErrorValue().GetLogString());
-						}
+						AsyncOp->SetError(MoveTemp(Error));
 						Promise.EmplaceValue();
 					});
 				}
@@ -227,18 +240,13 @@ TOnlineAsyncOpHandle<FAuthLogin> FAuthEOS::Login(FAuthLogin::Params&& Params)
 		{
 			FOnlineError CopyUserInfoError(Errors::FromEOSResult(CopyBestDisplayNameResult));
 			UE_LOG(LogOnlineServices, Warning, TEXT("[FAuthEOS::Login] Failure: EOS_UserInfo_CopyBestDisplayName %s"), *CopyUserInfoError.GetLogString());
-			InAsyncOp.SetError(Errors::Unknown(MoveTemp(CopyUserInfoError)));
 
 			TPromise<void> Promise;
 			TFuture<void> Future = Promise.GetFuture();
-
 			LogoutEASImpl(FAuthLogoutEASImpl::Params{ AccountInfoEOS->EpicAccountId })
-			.Next([Promise = MoveTemp(Promise)](TDefaultErrorResult<FAuthLogoutEASImpl>&& LogoutResult) mutable -> void
+			.Next([AsyncOp = InAsyncOp.AsShared(), Error = MoveTemp(CopyUserInfoError), Promise = MoveTemp(Promise)](TDefaultErrorResult<FAuthLogoutEASImpl>&&) mutable -> void
 			{
-				if (LogoutResult.IsError())
-				{
-					UE_LOG(LogOnlineServices, Warning, TEXT("[FAuthEOS::Login] Failure: LogoutEASImpl %s"), *LogoutResult.GetErrorValue().GetLogString());
-				}
+				AsyncOp->SetError(MoveTemp(Error));
 				Promise.EmplaceValue();
 			});
 
@@ -264,6 +272,274 @@ TOnlineAsyncOpHandle<FAuthLogin> FAuthEOS::Login(FAuthLogin::Params&& Params)
 	return Op->GetHandle();
 }
 
+TOnlineAsyncOpHandle<FAuthLinkAccount> FAuthEOS::LinkAccount(FAuthLinkAccount::Params&& Params)
+{
+	TOnlineAsyncOpRef<FAuthLinkAccount> Op = GetOp<FAuthLinkAccount>(MoveTemp(Params));
+	// Step 1: Set up operation data.
+	Op->Then([this](TOnlineAsyncOp<FAuthLinkAccount>& InAsyncOp)
+	{
+		const FAuthLinkAccount::Params& Params = InAsyncOp.GetParams();
+
+		// Check that user is valid.
+		if (!Params.PlatformUserId.IsValid())
+		{
+			InAsyncOp.SetError(Errors::InvalidParams());
+			return;
+		}
+
+		// Check that user scoped data exists for user.
+		FUserScopedData* UserData = FAuthEOSGS::GetUserScopedData(Params.PlatformUserId);
+		if (UserData == nullptr)
+		{
+			InAsyncOp.SetError(Errors::InvalidParams());
+			return;
+		}
+
+		// Make sure continuation exists.
+		const FLoginContinuationData* ContinuanceData = Algo::FindByPredicate(UserData->LoginContinuations, [&Params](const FLoginContinuationData& Continuation)->bool { return Continuation.ContinuationId == Params.ContinuationId; });
+		if (ContinuanceData == nullptr)
+		{
+			InAsyncOp.SetError(Errors::InvalidParams());
+			return;
+		}
+
+		InAsyncOp.Data.Set<FLoginContinuationData>(UE_ONLINE_AUTH_EOS_CONTINUANCE_DATA_KEY_NAME, *ContinuanceData);
+	})
+	// Step 2: Call link account.
+	.Then([this](TOnlineAsyncOp<FAuthLinkAccount>& InAsyncOp, TPromise<const EOS_Auth_LinkAccountCallbackInfo*>&& Promise)
+	{
+		const FAuthLinkAccount::Params& Params = InAsyncOp.GetParams();
+		const FLoginContinuationData& LoginContinuationData = GetOpDataChecked<FLoginContinuationData>(InAsyncOp, UE_ONLINE_AUTH_EOS_CONTINUANCE_DATA_KEY_NAME);
+		TSharedPtr<FAccountInfoEOS> AccountInfoEOS = AccountInfoRegistryEOS.Find(Params.PlatformUserId);
+
+		EOS_Auth_LinkAccountOptions LinkAccountOptions = {};
+		LinkAccountOptions.ApiVersion = 1;
+		LinkAccountOptions.ContinuanceToken = LoginContinuationData.ContinuanceToken;
+		LinkAccountOptions.LinkAccountFlags = LoginContinuationData.LinkAccountFlags;
+		LinkAccountOptions.LocalUserId = AccountInfoEOS ? AccountInfoEOS->EpicAccountId : nullptr;
+		UE_EOS_CHECK_API_MISMATCH(EOS_AUTH_LINKACCOUNT_API_LATEST, 1);
+
+		EOS_Async(EOS_Auth_LinkAccount, AuthHandle, LinkAccountOptions, MoveTemp(Promise));
+	})
+	// Step 3: Handle link account result.
+	.Then([this](TOnlineAsyncOp<FAuthLinkAccount>& InAsyncOp, const EOS_Auth_LinkAccountCallbackInfo* Data)
+	{
+		UE_LOG(LogOnlineServices, Verbose, TEXT("[FAuthEOS::LinkAccount] EOS_Auth_LinkAccount Result: [%s]"), *LexToString(Data->ResultCode));
+
+		if (Data->ResultCode == EOS_EResult::EOS_Success)
+		{
+			InAsyncOp.Data.Set<EOS_EpicAccountId>(UE_ONLINE_AUTH_EOS_SELECTED_ACCOUNT_ID_KEY_NAME, Data->SelectedAccountId);
+		}
+		else
+		{
+			InAsyncOp.SetError(Errors::FromEOSResult(Data->ResultCode));
+		}
+	})
+	// Step 4: Link account success handling / user info setup.
+	.Then([this](TOnlineAsyncOp<FAuthLinkAccount>& InAsyncOp)
+	{
+		const FAuthLinkAccount::Params& Params = InAsyncOp.GetParams();
+
+		// Remove continuance token.
+		FUserScopedData* UserData = FAuthEOSGS::GetUserScopedData(Params.PlatformUserId);
+		check(UserData);
+		UserData->LastLoginContinuationId = FLoginContinuationId();
+		UserData->LoginContinuations.SetNum(Algo::RemoveIf(UserData->LoginContinuations, [&Params](const FLoginContinuationData& Continuation)
+		{
+			return Continuation.ContinuationId == Params.ContinuationId;
+		}));
+
+		// Create or fetch FAccountInfoEOS for user.
+		const EOS_EpicAccountId SelectedUserAccount = GetOpDataChecked<EOS_EpicAccountId>(InAsyncOp, UE_ONLINE_AUTH_EOS_SELECTED_ACCOUNT_ID_KEY_NAME);
+		TSharedPtr<FAccountInfoEOS> AccountInfoEOS = AccountInfoRegistryEOS.Find(Params.PlatformUserId);
+		const bool bWasLoggedIn = AccountInfoEOS && AccountInfoEOS->EpicAccountId == SelectedUserAccount;
+		const bool bChoseDifferentAccount = bWasLoggedIn && AccountInfoEOS->EpicAccountId != SelectedUserAccount;
+
+		// Notify logout of previous account.
+		if (AccountInfoEOS && bChoseDifferentAccount)
+		{
+			AccountInfoEOS->LoginStatus = ELoginStatus::NotLoggedIn;
+			OnAuthLoginStatusChangedEvent.Broadcast(FAuthLoginStatusChanged{ AccountInfoEOS.ToSharedRef(), AccountInfoEOS->LoginStatus });
+			AccountInfoRegistryEOS.Unregister(AccountInfoEOS->AccountId);
+			AccountInfoEOS = nullptr;
+		}
+
+		// Setup user account if user is not already logged in.
+		if (AccountInfoEOS == nullptr)
+		{
+			AccountInfoEOS = MakeShared<FAccountInfoEOS>();
+			AccountInfoEOS->PlatformUserId = Params.PlatformUserId;
+			AccountInfoEOS->EpicAccountId = SelectedUserAccount;
+			AccountInfoEOS->LoginStatus = ELoginStatus::NotLoggedIn;
+		}
+
+		// Set user auth data on operation.
+		InAsyncOp.Data.Set<TSharedRef<FAccountInfoEOS>>(UE_ONLINE_AUTH_EOS_ACCOUNT_INFO_KEY_NAME, AccountInfoEOS.ToSharedRef());
+	})
+	// Step 5: Fetch external auth credentials for connect login.
+	.Then([this](TOnlineAsyncOp<FAuthLinkAccount>& InAsyncOp)
+	{
+		const FAuthLinkAccount::Params& Params = InAsyncOp.GetParams();
+		const TSharedRef<FAccountInfoEOS>& AccountInfoEOS = GetOpDataChecked<TSharedRef<FAccountInfoEOS>>(InAsyncOp, UE_ONLINE_AUTH_EOS_ACCOUNT_INFO_KEY_NAME);
+
+		TPromise<FAuthLoginConnectImpl::Params> Promise;
+		TFuture<FAuthLoginConnectImpl::Params> Future = Promise.GetFuture();
+
+		if (AccountInfoEOS->LoginStatus == ELoginStatus::NotLoggedIn)
+		{
+			TDefaultErrorResult<FAuthGetExternalAuthTokenImpl> AuthTokenResult = GetExternalAuthTokenImpl(FAuthGetExternalAuthTokenImpl::Params{AccountInfoEOS->EpicAccountId});
+			if (AuthTokenResult.IsError())
+			{
+				UE_LOG(LogOnlineServices, Warning, TEXT("[FAuthEOS::LinkAccount] Failure: GetExternalAuthTokenImpl %s"), *AuthTokenResult.GetErrorValue().GetLogString());
+
+				// Failed to acquire token - logout EAS.
+				LogoutEASImpl(FAuthLogoutEASImpl::Params{ AccountInfoEOS->EpicAccountId })
+				.Next([AsyncOp = InAsyncOp.AsShared(), Error = MoveTemp(AuthTokenResult.GetErrorValue()), Promise = MoveTemp(Promise)](TDefaultErrorResult<FAuthLogoutEASImpl>&&) mutable -> void
+				{
+					AsyncOp->SetError(MoveTemp(Error));
+					Promise.EmplaceValue(FAuthLoginConnectImpl::Params{});
+				});
+
+				return Future;
+			}
+
+			Promise.EmplaceValue(FAuthLoginConnectImpl::Params{Params.PlatformUserId, MoveTemp(AuthTokenResult.GetOkValue().Token)});
+		}
+		else
+		{
+			// No connect login is needed - user is already logged in.
+			Promise.EmplaceValue(FAuthLoginConnectImpl::Params{});
+		}
+		return Future;
+	})
+	// Step 6: Attempt connect login. On connect login failure handle logout of EAS.
+	.Then([this](TOnlineAsyncOp<FAuthLinkAccount>& InAsyncOp, FAuthLoginConnectImpl::Params&& LoginConnectParams)
+	{
+		const TSharedRef<FAccountInfoEOS>& AccountInfoEOS = GetOpDataChecked<TSharedRef<FAccountInfoEOS>>(InAsyncOp, UE_ONLINE_AUTH_EOS_ACCOUNT_INFO_KEY_NAME);
+
+		TPromise<void> Promise;
+		TFuture<void> Future = Promise.GetFuture();
+
+		if (AccountInfoEOS->LoginStatus == ELoginStatus::NotLoggedIn)
+		{
+			// Attempt connect login.
+			LoginConnectImpl(LoginConnectParams)
+			.Next([this, AccountInfoEOS, WeakOp = InAsyncOp.AsWeak(), Promise = MoveTemp(Promise)](TDefaultErrorResult<FAuthLoginConnectImpl>&& LoginResult) mutable -> void
+			{
+				if (TSharedPtr<TOnlineAsyncOp<FAuthLinkAccount>> AsyncOp = WeakOp.Pin())
+				{
+					if (LoginResult.IsError())
+					{
+						UE_LOG(LogOnlineServices, Warning, TEXT("[FAuthEOS::Login] Failure: LoginConnectImpl %s"), *LoginResult.GetErrorValue().GetLogString());
+
+						LogoutEASImpl(FAuthLogoutEASImpl::Params{ AccountInfoEOS->EpicAccountId })
+						.Next([AsyncOp, Error = MoveTemp(LoginResult.GetErrorValue()), Promise = MoveTemp(Promise)](TDefaultErrorResult<FAuthLogoutEASImpl>&&) mutable -> void
+						{
+							AsyncOp->SetError(MoveTemp(Error));
+							Promise.EmplaceValue();
+						});
+					}
+					else
+					{
+						// Successful login.
+						AccountInfoEOS->ProductUserId = LoginResult.GetOkValue().ProductUserId;
+						Promise.EmplaceValue();
+					}
+				}
+				else
+				{
+					Promise.EmplaceValue();
+				}
+			});
+		}
+		else
+		{
+			// No connect login is needed - user is already logged in.
+			Promise.EmplaceValue();
+		}
+
+		return Future;
+	})
+	// Step 7: Fetch dependent data.
+	.Then([this](TOnlineAsyncOp<FAuthLinkAccount>& InAsyncOp)
+	{
+		const TSharedRef<FAccountInfoEOS>& AccountInfoEOS = GetOpDataChecked<TSharedRef<FAccountInfoEOS>>(InAsyncOp, UE_ONLINE_AUTH_EOS_ACCOUNT_INFO_KEY_NAME);
+
+		if (AccountInfoEOS->LoginStatus == ELoginStatus::NotLoggedIn)
+		{
+			// Get display name
+			EOS_UserInfo_CopyBestDisplayNameOptions Options = {};
+			Options.ApiVersion = 1;
+			UE_EOS_CHECK_API_MISMATCH(EOS_USERINFO_COPYBESTDISPLAYNAME_API_LATEST, 1);
+			Options.LocalUserId = AccountInfoEOS->EpicAccountId;
+			Options.TargetUserId = AccountInfoEOS->EpicAccountId;
+
+			EOS_UserInfo_BestDisplayName* BestDisplayName;
+			EOS_EResult CopyBestDisplayNameResult = EOS_UserInfo_CopyBestDisplayName(UserInfoHandle, &Options, &BestDisplayName);
+
+			if (CopyBestDisplayNameResult == EOS_EResult::EOS_UserInfo_BestDisplayNameIndeterminate)
+			{
+				EOS_UserInfo_CopyBestDisplayNameWithPlatformOptions WithPlatformOptions = {};
+				WithPlatformOptions.ApiVersion = 1;
+				UE_EOS_CHECK_API_MISMATCH(EOS_USERINFO_COPYBESTDISPLAYNAMEWITHPLATFORM_API_LATEST, 1);
+				WithPlatformOptions.LocalUserId = AccountInfoEOS->EpicAccountId;
+				WithPlatformOptions.TargetUserId = AccountInfoEOS->EpicAccountId;
+				WithPlatformOptions.TargetPlatformType = EOS_OPT_Epic;
+
+				CopyBestDisplayNameResult = EOS_UserInfo_CopyBestDisplayNameWithPlatform(UserInfoHandle, &WithPlatformOptions, &BestDisplayName);
+			}
+
+			if (CopyBestDisplayNameResult == EOS_EResult::EOS_Success)
+			{
+				AccountInfoEOS->Attributes.Emplace(AccountAttributeData::DisplayName, *GetBestDisplayNameStr(*BestDisplayName));
+				EOS_UserInfo_BestDisplayName_Release(BestDisplayName);
+			}
+			else
+			{
+				FOnlineError CopyUserInfoError(Errors::FromEOSResult(CopyBestDisplayNameResult));
+				UE_LOG(LogOnlineServices, Warning, TEXT("[FAuthEOS::Login] Failure: EOS_UserInfo_CopyBestDisplayName %s"), *CopyUserInfoError.GetLogString());
+
+				TPromise<void> Promise;
+				TFuture<void> Future = Promise.GetFuture();
+				LogoutEASImpl(FAuthLogoutEASImpl::Params{ AccountInfoEOS->EpicAccountId })
+				.Next([AsyncOp = InAsyncOp.AsShared(), Error = MoveTemp(CopyUserInfoError), Promise = MoveTemp(Promise)](TDefaultErrorResult<FAuthLogoutEASImpl>&&) mutable -> void
+				{
+					AsyncOp->SetError(MoveTemp(Error));
+					Promise.EmplaceValue();
+				});
+
+				return Future;
+			}
+		}
+
+		return MakeFulfilledPromise<void>().GetFuture();
+	})
+	// Step 8: bookkeeping and notifications.
+	.Then([this](TOnlineAsyncOp<FAuthLinkAccount>& InAsyncOp)
+	{
+		const TSharedRef<FAccountInfoEOS>& AccountInfoEOS = GetOpDataChecked<TSharedRef<FAccountInfoEOS>>(InAsyncOp, UE_ONLINE_AUTH_EOS_ACCOUNT_INFO_KEY_NAME);
+		const bool bUserWasLoggedIn = AccountInfoEOS->LoginStatus == ELoginStatus::LoggedIn;
+
+		if (bUserWasLoggedIn)
+		{
+			UE_LOG(LogOnlineServices, Log, TEXT("[FAuthEOS::FAuthLinkAccount] Successfully linked account. AccountId: %s"), *ToLogString(AccountInfoEOS->AccountId));
+		}
+		else
+		{
+			AccountInfoEOS->LoginStatus = ELoginStatus::LoggedIn;
+			AccountInfoEOS->AccountId = CreateAccountId(AccountInfoEOS->EpicAccountId, AccountInfoEOS->ProductUserId);
+			AccountInfoRegistryEOS.Register(AccountInfoEOS);
+			UE_LOG(LogOnlineServices, Log, TEXT("[FAuthEOS::FAuthLinkAccount] Successfully logged in. AccountId: %s"), *ToLogString(AccountInfoEOS->AccountId));
+			OnAuthLoginStatusChangedEvent.Broadcast(FAuthLoginStatusChanged{ AccountInfoEOS, AccountInfoEOS->LoginStatus });
+		}
+
+		InAsyncOp.SetResult(FAuthLinkAccount::Result{AccountInfoEOS});
+	})
+	.Enqueue(GetSerialQueue());
+
+	return Op->GetHandle();
+}
+
 template <typename IdType>
 TFuture<FAccountId> ResolveAccountIdImpl(FAuthEOS& AuthEOS, const FAccountId& LocalAccountId, const IdType InId)
 {
@@ -349,7 +625,7 @@ TOnlineAsyncOpHandle<FAuthQueryExternalAuthToken> FAuthEOS::QueryExternalAuthTok
 				if (AuthTokenResult.IsError())
 				{
 					UE_LOG(LogOnlineServices, Warning, TEXT("[FAuthEOS::QueryExternalAuthToken] Failure: GetExternalAuthTokenImpl %s"), *AuthTokenResult.GetErrorValue().GetLogString());
-					InAsyncOp.SetError(Errors::Unknown(MoveTemp(AuthTokenResult.GetErrorValue())));
+					InAsyncOp.SetError(MoveTemp(AuthTokenResult.GetErrorValue()));
 					return;
 				}
 
@@ -393,6 +669,30 @@ TOnlineAsyncOpHandle<FAuthQueryExternalAuthToken> FAuthEOS::QueryExternalAuthTok
 	return Op->GetHandle();
 }
 
+TOnlineResult<FAuthGetLinkAccountContinuationId> FAuthEOS::GetLinkAccountContinuationId(FAuthGetLinkAccountContinuationId::Params&& Params) const
+{
+	if (!Params.PlatformUserId.IsValid())
+	{
+		return TOnlineResult<FAuthGetLinkAccountContinuationId>(Errors::InvalidUser());
+	}
+
+	const FUserScopedData* UserData = GetUserScopedData(Params.PlatformUserId);
+	if (UserData == nullptr)
+	{
+		UE_LOG(LogOnlineServices, Verbose, TEXT("[%hs]: Failed to find user scoped data. PlatformUserId: %s."), __FUNCTION__, *ToLogString(Params.PlatformUserId));
+		return TOnlineResult<FAuthGetLinkAccountContinuationId>(Errors::NotFound());
+	}
+
+	if (!UserData->LastLoginContinuationId.IsValid())
+	{
+		UE_LOG(LogOnlineServices, Verbose, TEXT("[%hs]: Failed to find valid login continuation. PlatformUserId: %s."), __FUNCTION__, *ToLogString(Params.PlatformUserId));
+		return TOnlineResult<FAuthGetLinkAccountContinuationId>(Errors::NotFound());
+	}
+
+	UE_LOG(LogOnlineServices, Verbose, TEXT("[%hs]: Found continuation id. PlatformUserId: %s, ContinuationId: %s."), __FUNCTION__, *ToLogString(Params.PlatformUserId), *ToLogString(UserData->LastLoginContinuationId));
+	return TOnlineResult<FAuthGetLinkAccountContinuationId>(FAuthGetLinkAccountContinuationId::Result{UserData->LastLoginContinuationId});
+}
+
 TFuture<FAccountId> FAuthEOS::ResolveAccountId(const FAccountId& LocalAccountId, const EOS_EpicAccountId EpicAccountId)
 {
 	return ResolveAccountIdImpl(*this, LocalAccountId, EpicAccountId);
diff --git a/Engine/Plugins/Online/OnlineServicesEOS/Source/Public/Online/AuthEOS.h b/Engine/Plugins/Online/OnlineServicesEOS/Source/Public/Online/AuthEOS.h
--- a/Engine/Plugins/Online/OnlineServicesEOS/Source/Public/Online/AuthEOS.h
+++ b/Engine/Plugins/Online/OnlineServicesEOS/Source/Public/Online/AuthEOS.h
@@ -24,8 +24,10 @@ public:
 
 	// Begin IAuth
 	virtual TOnlineAsyncOpHandle<FAuthLogin> Login(FAuthLogin::Params&& Params) override;
+	virtual TOnlineAsyncOpHandle<FAuthLinkAccount> LinkAccount(FAuthLinkAccount::Params&& Params) override;
 	virtual TOnlineAsyncOpHandle<FAuthQueryExternalServerAuthTicket> QueryExternalServerAuthTicket(FAuthQueryExternalServerAuthTicket::Params&& Params) override;
 	virtual TOnlineAsyncOpHandle<FAuthQueryExternalAuthToken> QueryExternalAuthToken(FAuthQueryExternalAuthToken::Params&& Params) override;
+	virtual TOnlineResult<FAuthGetLinkAccountContinuationId> GetLinkAccountContinuationId(FAuthGetLinkAccountContinuationId::Params&& Params) const;
 	// End IAuth
 
 	// Begin FAuthEOSGS
diff --git a/Engine/Plugins/Online/OnlineServicesEOSGS/Source/Private/Online/AuthEOSGS.cpp b/Engine/Plugins/Online/OnlineServicesEOSGS/Source/Private/Online/AuthEOSGS.cpp
--- a/Engine/Plugins/Online/OnlineServicesEOSGS/Source/Private/Online/AuthEOSGS.cpp
+++ b/Engine/Plugins/Online/OnlineServicesEOSGS/Source/Private/Online/AuthEOSGS.cpp
@@ -308,6 +308,7 @@ TDefaultErrorResultInternal<FEOSAuthLoginOptions> FEOSAuthLoginOptions::CreateIm
 		FCStringAnsi::Strncpy(EOSAuthLoginOptions.TokenUtf8.GetData(), TCHAR_TO_UTF8(*ExternalAuthToken.Data), EOSAuthLoginOptions.TokenUtf8.Num());
 		EOSAuthLoginOptions.CredentialsData.Token = EOSAuthLoginOptions.TokenUtf8.GetData();
 		EOSAuthLoginOptions.CredentialsData.ExternalType = ExternalAuthTranslationTraits->Type;
+		EOSAuthLoginOptions.LinkAccountFlags = ExternalAuthTranslationTraits->LinkAccountFlags;
 
 		UE_LOG(LogOnlineServices, VeryVerbose, TEXT("FEOSAuthLoginOptions::Create: Using token type: %s, Token data: %.*s"), *ExternalAuthToken.Type.ToString(), ExternalAuthToken.Data.Len(), *ExternalAuthToken.Data);
 	}
@@ -466,6 +467,14 @@ TOnlineAsyncOpHandle<FAuthLogin> FAuthEOSGS::Login(FAuthLogin::Params&& Params)
 	Op->Then([this](TOnlineAsyncOp<FAuthLogin>& InAsyncOp)
 	{
 		const FAuthLogin::Params& Params = InAsyncOp.GetParams();
+
+		// Check that user is valid.
+		if (!Params.PlatformUserId.IsValid())
+		{
+			InAsyncOp.SetError(Errors::InvalidParams());
+			return;
+		}
+
 		TSharedPtr<FAccountInfoEOS> AccountInfoEOS = AccountInfoRegistryEOS.Find(Params.PlatformUserId);
 		if (!AccountInfoEOS)
 		{
@@ -543,19 +552,14 @@ TOnlineAsyncOpHandle<FAuthLogin> FAuthEOSGS::Login(FAuthLogin::Params&& Params)
 			if (AuthTokenResult.IsError())
 			{
 				UE_LOG(LogOnlineServices, Warning, TEXT("[FAuthEOSGS::Login] Failure: GetExternalAuthTokenImpl %s"), *AuthTokenResult.GetErrorValue().GetLogString());
-				InAsyncOp.SetError(Errors::Unknown(MoveTemp(AuthTokenResult.GetErrorValue())));
 
 				// Failed to acquire token - logout EAS.
 				LogoutEASImpl(FAuthLogoutEASImpl::Params{ AccountInfoEOS->EpicAccountId })
-				.Next([Promise = MoveTemp(Promise)](TDefaultErrorResult<FAuthLogoutEASImpl>&& LogoutResult) mutable -> void
+				.Next([AsyncOp = InAsyncOp.AsShared(), Error = MoveTemp(AuthTokenResult.GetErrorValue()), Promise = MoveTemp(Promise)](TDefaultErrorResult<FAuthLogoutEASImpl>&&) mutable -> void
 				{
-					if (LogoutResult.IsError())
-					{
-						UE_LOG(LogOnlineServices, Warning, TEXT("[FAuthEOSGS::Login] Failure: LogoutEASImpl %s"), *LogoutResult.GetErrorValue().GetLogString());
-					}
+					AsyncOp->SetError(MoveTemp(Error));
 					Promise.EmplaceValue(FAuthLoginConnectImpl::Params{});
 				});
-
 				return Future;
 			}
 
@@ -586,23 +590,19 @@ TOnlineAsyncOpHandle<FAuthLogin> FAuthEOSGS::Login(FAuthLogin::Params&& Params)
 
 		// Attempt connect login.
 		LoginConnectImpl(LoginConnectParams)
-		.Next([this, AccountInfoEOS, Op = InAsyncOp.AsShared(), Promise = MoveTemp(Promise)](TDefaultErrorResult<FAuthLoginConnectImpl>&& LoginResult) mutable -> void
+		.Next([this, AccountInfoEOS, AsyncOp = InAsyncOp.AsShared(), Promise = MoveTemp(Promise)](TDefaultErrorResult<FAuthLoginConnectImpl>&& LoginResult) mutable -> void
 		{
 			if (LoginResult.IsError())
 			{
 				UE_LOG(LogOnlineServices, Warning, TEXT("[FAuthEOSGS::Login] Failure: LoginConnectImpl %s"), *LoginResult.GetErrorValue().GetLogString());
-				Op->SetError(Errors::Unknown(MoveTemp(LoginResult.GetErrorValue())));
 
 				// Logout of EAS on login failure if necessary.
 				if (AccountInfoEOS->EpicAccountId)
 				{
 					LogoutEASImpl(FAuthLogoutEASImpl::Params{ AccountInfoEOS->EpicAccountId })
-					.Next([Promise = MoveTemp(Promise)](TDefaultErrorResult<FAuthLogoutEASImpl>&& LogoutResult) mutable -> void
+					.Next([AsyncOp = AsyncOp->AsShared(), Error = MoveTemp(LoginResult.GetErrorValue()), Promise = MoveTemp(Promise)](TDefaultErrorResult<FAuthLogoutEASImpl>&&) mutable -> void
 					{
-						if (LogoutResult.IsError())
-						{
-							UE_LOG(LogOnlineServices, Warning, TEXT("[FAuthEOSGS::Login] Failure: LogoutEASImpl %s"), *LogoutResult.GetErrorValue().GetLogString());
-						}
+						AsyncOp->SetError(MoveTemp(Error));
 						Promise.EmplaceValue();
 					});
 				}
@@ -640,26 +640,24 @@ TOnlineAsyncOpHandle<FAuthLogin> FAuthEOSGS::Login(FAuthLogin::Params&& Params)
 			{
 				FOnlineError ProductUserIdError(Errors::FromEOSResult(ProductUserIdResult));
 				UE_LOG(LogOnlineServices, Warning, TEXT("[FAuthEOSGS::Login] Failure: EOS_ProductUserId_ToString %s"), *ProductUserIdError.GetLogString());
-				InAsyncOp.SetError(Errors::Unknown(MoveTemp(ProductUserIdError)));
 
 				// Handle EAS logout if needed.
 				if (AccountInfoEOS->EpicAccountId)
 				{
 					TPromise<void> Promise;
 					TFuture<void> Future = Promise.GetFuture();
-
 					LogoutEASImpl(FAuthLogoutEASImpl::Params{ AccountInfoEOS->EpicAccountId })
-					.Next([Promise = MoveTemp(Promise)](TDefaultErrorResult<FAuthLogoutEASImpl>&& LogoutResult) mutable -> void
+					.Next([AsyncOp = InAsyncOp.AsShared(), Error = MoveTemp(ProductUserIdError), Promise = MoveTemp(Promise)](TDefaultErrorResult<FAuthLogoutEASImpl>&&) mutable -> void
 					{
-						if (LogoutResult.IsError())
-						{
-							UE_LOG(LogOnlineServices, Warning, TEXT("[FAuthEOSGS::Login] Failure: LogoutEASImpl %s"), *LogoutResult.GetErrorValue().GetLogString());
-						}
+						AsyncOp->SetError(MoveTemp(Error));
 						Promise.EmplaceValue();
 					});
-
 					return Future;
 				}
+				else
+				{
+					InAsyncOp.SetError(MoveTemp(ProductUserIdError));
+				}
 			}
 		}
 
@@ -751,17 +749,11 @@ TOnlineAsyncOpHandle<FAuthLogout> FAuthEOSGS::Logout(FAuthLogout::Params&& Param
 		{
 			TPromise<void> Promise;
 			TFuture<void> Future = Promise.GetFuture();
-
 			LogoutEASImpl(FAuthLogoutEASImpl::Params{ AccountInfoEOS->EpicAccountId })
-			.Next([this, AccountInfoEOS, Promise = MoveTemp(Promise)](TDefaultErrorResult<FAuthLogoutEASImpl>&& LogoutResult) mutable -> void
+			.Next([Promise = MoveTemp(Promise)](TDefaultErrorResult<FAuthLogoutEASImpl>&&) mutable -> void
 			{
-				if (LogoutResult.IsError())
-				{
-					UE_LOG(LogOnlineServices, Warning, TEXT("[FAuthEOSGS::Logout] Failure: LogoutEASImpl %s"), *LogoutResult.GetErrorValue().GetLogString());
-				}
 				Promise.EmplaceValue();
 			});
-
 			return Future;
 		}
 		else
@@ -809,7 +801,7 @@ TOnlineAsyncOpHandle<FAuthBeginVerifiedAuthSession> FAuthEOSGS::BeginVerifiedAut
 	return Operation->GetHandle();
 }
 
-TOnlineResult<FAuthGetRelyingParty> FAuthEOSGS::GetRelyingParty() const
+TOnlineResult<FAuthGetRelyingParty> FAuthEOSGS::GetRelyingParty(FAuthGetRelyingParty::Params&& Params) const
 {
 	IEOSSDKManager* Manager = IEOSSDKManager::Get();
 	const FString& PlatformConfigName = GetServices<FOnlineServicesEOSGS>().GetEOSPlatformHandle()->GetConfigName();
@@ -890,12 +882,13 @@ TFuture<TDefaultErrorResult<FAuthLoginEASImpl>> FAuthEOSGS::LoginEASImpl(const F
 	}
 
 	const bool bIsPersistentAuthLogin = LoginParams.CredentialsType == LoginCredentialsType::PersistentAuth;
+	const EOS_ELinkAccountFlags LinkAccountFlags = LoginOptionsResult.GetOkValue().GetLinkAccountFlags();
 
 	TPromise<TDefaultErrorResult<FAuthLoginEASImpl>> Promise;
 	TFuture<TDefaultErrorResult<FAuthLoginEASImpl>> Future = Promise.GetFuture();
 
 	EOS_Async(EOS_Auth_Login, AuthHandle, MoveTemp(LoginOptionsResult.GetOkValue()),
-	[AuthHandle = AuthHandle, bIsPersistentAuthLogin, Promise = MoveTemp(Promise)](const EOS_Auth_LoginCallbackInfo* Data) mutable -> void
+	[this, AuthHandle = AuthHandle, bIsPersistentAuthLogin, PlatformUserId = LoginParams.PlatformUserId, LinkAccountFlags, bAutoLinkAccount = LoginParams.bAutoLinkAccount, Promise = MoveTemp(Promise)](const EOS_Auth_LoginCallbackInfo* Data) mutable -> void
 	{
 		UE_LOG(LogOnlineServices, Verbose, TEXT("[FAuthEOSGS::LoginEASImpl] EOS_Auth_Login Result: [%s]"), *LexToString(Data->ResultCode));
 
@@ -905,25 +898,44 @@ TFuture<TDefaultErrorResult<FAuthLoginEASImpl>> FAuthEOSGS::LoginEASImpl(const F
 		}
 		else if (Data->ResultCode == EOS_EResult::EOS_InvalidUser && Data->ContinuanceToken != nullptr)
 		{
-			EOS_Auth_LinkAccountOptions LinkAccountOptions = {};
-			LinkAccountOptions.ApiVersion = 1;
-			LinkAccountOptions.ContinuanceToken = Data->ContinuanceToken;
-			UE_EOS_CHECK_API_MISMATCH(EOS_AUTH_LINKACCOUNT_API_LATEST, 1);
-
-			EOS_Async(EOS_Auth_LinkAccount, AuthHandle, LinkAccountOptions,
-			[Promise = MoveTemp(Promise)](const EOS_Auth_LinkAccountCallbackInfo* Data) mutable -> void
+			if (bAutoLinkAccount)
 			{
-				UE_LOG(LogOnlineServices, Verbose, TEXT("[FAuthEOSGS::LoginEASImpl] EOS_Auth_LinkAccount Result: [%s]"), *LexToString(Data->ResultCode));
+				EOS_Auth_LinkAccountOptions LinkAccountOptions = {};
+				LinkAccountOptions.ApiVersion = 1;
+				LinkAccountOptions.ContinuanceToken = Data->ContinuanceToken;
+				LinkAccountOptions.LinkAccountFlags = LinkAccountFlags;
+				UE_EOS_CHECK_API_MISMATCH(EOS_AUTH_LINKACCOUNT_API_LATEST, 1);
+
+				EOS_Async(EOS_Auth_LinkAccount, AuthHandle, LinkAccountOptions,
+				[Promise = MoveTemp(Promise)](const EOS_Auth_LinkAccountCallbackInfo* Data) mutable -> void
+				{
+					UE_LOG(LogOnlineServices, Verbose, TEXT("[FAuthEOSGS::LoginEASImpl] EOS_Auth_LinkAccount Result: [%s]"), *LexToString(Data->ResultCode));
 
-				if (Data->ResultCode == EOS_EResult::EOS_Success)
+					if (Data->ResultCode == EOS_EResult::EOS_Success)
+					{
+						Promise.SetValue(TDefaultErrorResult<FAuthLoginEASImpl>(FAuthLoginEASImpl::Result{ Data->LocalUserId }));
+					}
+					else
+					{
+						Promise.SetValue(TDefaultErrorResult<FAuthLoginEASImpl>(Errors::FromEOSResult(Data->ResultCode)));
+					}
+				});
+			}
+			else
+			{
+				UE_LOG(LogOnlineServices, Verbose, TEXT("[FAuthEOSGS::LoginEASImpl] Saving link account continuance token for manual link account usage."));
+				if (FUserScopedData* UserData = GetOrCreateUserScopedData(PlatformUserId))
 				{
-					Promise.SetValue(TDefaultErrorResult<FAuthLoginEASImpl>(FAuthLoginEASImpl::Result{ Data->LocalUserId }));
+					UserData->LastLoginContinuationId = FLoginContinuationId(EOnlineServices::Epic, NextLoginContinuationId++);
+					UserData->LoginContinuations.Add(FLoginContinuationData{UserData->LastLoginContinuationId, Data->ContinuanceToken, LinkAccountFlags});
 				}
 				else
 				{
-					Promise.SetValue(TDefaultErrorResult<FAuthLoginEASImpl>(Errors::FromEOSResult(Data->ResultCode)));
+					UE_LOG(LogOnlineServices, Warning, TEXT("[FAuthEOSGS::LoginEASImpl] Failed to access user scoped data."));
 				}
-			});
+
+				Promise.SetValue(TDefaultErrorResult<FAuthLoginEASImpl>(Errors::FromEOSResult(Data->ResultCode)));
+			}
 		}
 		else
 		{
@@ -976,14 +988,14 @@ TFuture<TDefaultErrorResult<FAuthLogoutEASImpl>> FAuthEOSGS::LogoutEASImpl(const
 	EOS_Async(EOS_Auth_Logout, AuthHandle, LogoutOptions,
 	[Promise = MoveTemp(Promise)](const EOS_Auth_LogoutCallbackInfo* Data) mutable -> void
 	{
-		UE_LOG(LogOnlineServices, Verbose, TEXT("[FAuthEOSGS::LogoutEASImpl] EOS_Auth_Logout Result: [%s]"), *LexToString(Data->ResultCode));
-
 		if (Data->ResultCode == EOS_EResult::EOS_Success)
 		{
+			UE_LOG(LogOnlineServices, Verbose, TEXT("[FAuthEOSGS::LogoutEASImpl] EOS_Auth_Logout Result: [%s]"), *LexToString(Data->ResultCode));
 			Promise.SetValue(TDefaultErrorResult<FAuthLogoutEASImpl>(FAuthLogoutEASImpl::Result{}));
 		}
 		else
 		{
+			UE_LOG(LogOnlineServices, Warning, TEXT("[FAuthEOSGS::LogoutEASImpl] EOS_Auth_Logout Result: [%s]"), *LexToString(Data->ResultCode));
 			Promise.SetValue(TDefaultErrorResult<FAuthLogoutEASImpl>(Errors::FromEOSResult(Data->ResultCode)));
 		}
 	});
@@ -1108,7 +1120,7 @@ TOnlineAsyncOpHandle<FAuthConnectLoginRecoveryImpl> FAuthEOSGS::ConnectLoginReco
 		if (AuthTokenResult.IsError())
 		{
 			UE_LOG(LogOnlineServices, Warning, TEXT("[FAuthEOSGS::ConnectLoginRecoveryImplOp] Failure: GetExternalAuthTokenImpl %s"), *AuthTokenResult.GetErrorValue().GetLogString());
-			InAsyncOp.SetError(Errors::Unknown(MoveTemp(AuthTokenResult.GetErrorValue())));
+			InAsyncOp.SetError(MoveTemp(AuthTokenResult.GetErrorValue()));
 
 			// Reinitialize recovery timer.
 			InitializeConnectLoginRecoveryTimer(AccountInfoEOS);
@@ -1132,7 +1144,7 @@ TOnlineAsyncOpHandle<FAuthConnectLoginRecoveryImpl> FAuthEOSGS::ConnectLoginReco
 			if (LoginResult.IsError())
 			{
 				UE_LOG(LogOnlineServices, Warning, TEXT("[FAuthEOSGS::ConnectLoginRecoveryImplOp] Failure: LoginConnectImpl %s"), *LoginResult.GetErrorValue().GetLogString());
-				Op->SetError(Errors::Unknown(MoveTemp(LoginResult.GetErrorValue())));
+				Op->SetError(MoveTemp(LoginResult.GetErrorValue()));
 
 				// Reinitialize recovery timer.
 				InitializeConnectLoginRecoveryTimer(AccountInfoEOS);
@@ -1447,4 +1459,32 @@ FAccountId FAuthEOSGS::CreateAccountId(const EOS_ProductUserId ProductUserId)
 	return FOnlineAccountIdRegistryEOSGS::Get().FindOrAddAccountId(ProductUserId);
 }
 
+FAuthEOSGS::FUserScopedData* FAuthEOSGS::GetUserScopedData(FPlatformUserId PlatformUserId)
+{
+	const int32 UserIndex = FPlatformMisc::GetUserIndexForPlatformUser(PlatformUserId);
+	return UserScopedData.IsValidIndex(UserIndex) ? &UserScopedData[UserIndex] : nullptr;
+}
+
+const FAuthEOSGS::FUserScopedData* FAuthEOSGS::GetUserScopedData(FPlatformUserId PlatformUserId) const
+{
+	const int32 UserIndex = FPlatformMisc::GetUserIndexForPlatformUser(PlatformUserId);
+	return UserScopedData.IsValidIndex(UserIndex) ? &UserScopedData[UserIndex] : nullptr;
+}
+
+FAuthEOSGS::FUserScopedData* FAuthEOSGS::GetOrCreateUserScopedData(FPlatformUserId PlatformUserId)
+{
+	const int32 UserIndex = FPlatformMisc::GetUserIndexForPlatformUser(PlatformUserId);
+	if (UserIndex == INDEX_NONE)
+	{
+		return nullptr;
+	}
+
+	if (!UserScopedData.IsValidIndex(UserIndex))
+	{
+		UserScopedData.Insert(UserIndex, FUserScopedData());
+	}
+
+	return &UserScopedData[UserIndex];
+}
+
 /* UE::Online */ }
diff --git a/Engine/Plugins/Online/OnlineServicesEOSGS/Source/Private/Online/EOSAuthLoginOptionsCommon.cpp b/Engine/Plugins/Online/OnlineServicesEOSGS/Source/Private/Online/EOSAuthLoginOptionsCommon.cpp
--- a/Engine/Plugins/Online/OnlineServicesEOSGS/Source/Private/Online/EOSAuthLoginOptionsCommon.cpp
+++ b/Engine/Plugins/Online/OnlineServicesEOSGS/Source/Private/Online/EOSAuthLoginOptionsCommon.cpp
@@ -6,6 +6,12 @@
 
 namespace UE::Online {
 
+namespace CompatibilityLoginCredentialsType
+{
+	const FName Password = TEXT("epic");
+	const FName Developer = TEXT("dev_tool");
+}
+
 const FEOSAuthTranslationTraits* FEOSAuthLoginOptionsCommon::GetLoginTranslatorTraits(FName Name)
 {
 	static const TMap<FName, FEOSAuthTranslationTraits> SupportedLoginTranslatorTraits = {
@@ -16,6 +22,8 @@ const FEOSAuthTranslationTraits* FEOSAuthLoginOptionsCommon::GetLoginTranslatorT
 		{ LoginCredentialsType::RefreshToken, { EOS_ELoginCredentialType::EOS_LCT_RefreshToken, EEOSAuthTranslationFlags::SetTokenFromString } },
 		{ LoginCredentialsType::AccountPortal, { EOS_ELoginCredentialType::EOS_LCT_AccountPortal, EEOSAuthTranslationFlags::SetId | EEOSAuthTranslationFlags::SetTokenFromString } },
 		{ LoginCredentialsType::ExternalAuth, { EOS_ELoginCredentialType::EOS_LCT_ExternalAuth, EEOSAuthTranslationFlags::SetTokenFromExternalAuth } },
+		{ CompatibilityLoginCredentialsType::Password, { EOS_ELoginCredentialType::EOS_LCT_Password, EEOSAuthTranslationFlags::SetId | EEOSAuthTranslationFlags::SetTokenFromString } },
+		{ CompatibilityLoginCredentialsType::Developer, { EOS_ELoginCredentialType::EOS_LCT_Developer, EEOSAuthTranslationFlags::SetId | EEOSAuthTranslationFlags::SetTokenFromString } },
 	};
 
 	return SupportedLoginTranslatorTraits.Find(Name);
@@ -24,24 +32,24 @@ const FEOSAuthTranslationTraits* FEOSAuthLoginOptionsCommon::GetLoginTranslatorT
 const FEOSExternalAuthTranslationTraits* FEOSAuthLoginOptionsCommon::GetExternalAuthTranslationTraits(FName ExternalAuthType)
 {
 	static const TMap<FName, FEOSExternalAuthTranslationTraits> SupportedExternalAuthTraits = {
-		{ ExternalLoginType::Epic, { EOS_EExternalCredentialType::EOS_ECT_EPIC } },
-		{ ExternalLoginType::SteamSessionTicket, { EOS_EExternalCredentialType::EOS_ECT_STEAM_SESSION_TICKET } },
-		{ ExternalLoginType::PsnIdToken, { EOS_EExternalCredentialType::EOS_ECT_PSN_ID_TOKEN } },
-		{ ExternalLoginType::XblXstsToken, { EOS_EExternalCredentialType::EOS_ECT_XBL_XSTS_TOKEN } },
-		{ ExternalLoginType::DiscordAccessToken, { EOS_EExternalCredentialType::EOS_ECT_DISCORD_ACCESS_TOKEN } },
-		{ ExternalLoginType::GogSessionTicket, { EOS_EExternalCredentialType::EOS_ECT_GOG_SESSION_TICKET } },
-		{ ExternalLoginType::NintendoIdToken, { EOS_EExternalCredentialType::EOS_ECT_NINTENDO_ID_TOKEN } },
-		{ ExternalLoginType::NintendoNsaIdToken, { EOS_EExternalCredentialType::EOS_ECT_NINTENDO_NSA_ID_TOKEN } },
-		{ ExternalLoginType::UplayAccessToken, { EOS_EExternalCredentialType::EOS_ECT_UPLAY_ACCESS_TOKEN } },
-		{ ExternalLoginType::OpenIdAccessToken, { EOS_EExternalCredentialType::EOS_ECT_OPENID_ACCESS_TOKEN } },
-		{ ExternalLoginType::DeviceIdAccessToken, { EOS_EExternalCredentialType::EOS_ECT_DEVICEID_ACCESS_TOKEN } },
-		{ ExternalLoginType::AppleIdToken, { EOS_EExternalCredentialType::EOS_ECT_APPLE_ID_TOKEN } },
-		{ ExternalLoginType::GoogleIdToken, { EOS_EExternalCredentialType::EOS_ECT_GOOGLE_ID_TOKEN } },
-		{ ExternalLoginType::OculusUserIdNonce, { EOS_EExternalCredentialType::EOS_ECT_OCULUS_USERID_NONCE } },
-		{ ExternalLoginType::ItchioJwt, { EOS_EExternalCredentialType::EOS_ECT_ITCHIO_JWT } },
-		{ ExternalLoginType::ItchioKey, { EOS_EExternalCredentialType::EOS_ECT_ITCHIO_KEY } },
-		{ ExternalLoginType::EpicIdToken, { EOS_EExternalCredentialType::EOS_ECT_EPIC_ID_TOKEN } },
-		{ ExternalLoginType::AmazonAccessToken, { EOS_EExternalCredentialType::EOS_ECT_AMAZON_ACCESS_TOKEN } },
+		{ ExternalLoginType::Epic, { EOS_EExternalCredentialType::EOS_ECT_EPIC, EOS_ELinkAccountFlags::EOS_LA_NoFlags } },
+		{ ExternalLoginType::SteamSessionTicket, { EOS_EExternalCredentialType::EOS_ECT_STEAM_SESSION_TICKET, EOS_ELinkAccountFlags::EOS_LA_NoFlags } },
+		{ ExternalLoginType::PsnIdToken, { EOS_EExternalCredentialType::EOS_ECT_PSN_ID_TOKEN, EOS_ELinkAccountFlags::EOS_LA_NoFlags } },
+		{ ExternalLoginType::XblXstsToken, { EOS_EExternalCredentialType::EOS_ECT_XBL_XSTS_TOKEN, EOS_ELinkAccountFlags::EOS_LA_NoFlags } },
+		{ ExternalLoginType::DiscordAccessToken, { EOS_EExternalCredentialType::EOS_ECT_DISCORD_ACCESS_TOKEN, EOS_ELinkAccountFlags::EOS_LA_NoFlags } },
+		{ ExternalLoginType::GogSessionTicket, { EOS_EExternalCredentialType::EOS_ECT_GOG_SESSION_TICKET, EOS_ELinkAccountFlags::EOS_LA_NoFlags } },
+		{ ExternalLoginType::NintendoIdToken, { EOS_EExternalCredentialType::EOS_ECT_NINTENDO_ID_TOKEN, EOS_ELinkAccountFlags::EOS_LA_NoFlags } },
+		{ ExternalLoginType::NintendoNsaIdToken, { EOS_EExternalCredentialType::EOS_ECT_NINTENDO_NSA_ID_TOKEN, EOS_ELinkAccountFlags::EOS_LA_NintendoNsaId } },
+		{ ExternalLoginType::UplayAccessToken, { EOS_EExternalCredentialType::EOS_ECT_UPLAY_ACCESS_TOKEN, EOS_ELinkAccountFlags::EOS_LA_NoFlags } },
+		{ ExternalLoginType::OpenIdAccessToken, { EOS_EExternalCredentialType::EOS_ECT_OPENID_ACCESS_TOKEN, EOS_ELinkAccountFlags::EOS_LA_NoFlags } },
+		{ ExternalLoginType::DeviceIdAccessToken, { EOS_EExternalCredentialType::EOS_ECT_DEVICEID_ACCESS_TOKEN, EOS_ELinkAccountFlags::EOS_LA_NoFlags } },
+		{ ExternalLoginType::AppleIdToken, { EOS_EExternalCredentialType::EOS_ECT_APPLE_ID_TOKEN, EOS_ELinkAccountFlags::EOS_LA_NoFlags } },
+		{ ExternalLoginType::GoogleIdToken, { EOS_EExternalCredentialType::EOS_ECT_GOOGLE_ID_TOKEN, EOS_ELinkAccountFlags::EOS_LA_NoFlags } },
+		{ ExternalLoginType::OculusUserIdNonce, { EOS_EExternalCredentialType::EOS_ECT_OCULUS_USERID_NONCE, EOS_ELinkAccountFlags::EOS_LA_NoFlags } },
+		{ ExternalLoginType::ItchioJwt, { EOS_EExternalCredentialType::EOS_ECT_ITCHIO_JWT, EOS_ELinkAccountFlags::EOS_LA_NoFlags } },
+		{ ExternalLoginType::ItchioKey, { EOS_EExternalCredentialType::EOS_ECT_ITCHIO_KEY, EOS_ELinkAccountFlags::EOS_LA_NoFlags } },
+		{ ExternalLoginType::EpicIdToken, { EOS_EExternalCredentialType::EOS_ECT_EPIC_ID_TOKEN, EOS_ELinkAccountFlags::EOS_LA_NoFlags } },
+		{ ExternalLoginType::AmazonAccessToken, { EOS_EExternalCredentialType::EOS_ECT_AMAZON_ACCESS_TOKEN, EOS_ELinkAccountFlags::EOS_LA_NoFlags } },
 	};
 
 	return SupportedExternalAuthTraits.Find(ExternalAuthType);
@@ -66,6 +74,7 @@ FEOSAuthLoginOptionsCommon& FEOSAuthLoginOptionsCommon::operator=(FEOSAuthLoginO
 	{
 		TokenUtf8 = MoveTemp(Other.TokenUtf8);
 		CredentialsData.Token = TokenUtf8.GetData();
+		LinkAccountFlags = Other.LinkAccountFlags;
 	}
 	if (CredentialsData.SystemAuthCredentialsOptions)
 	{
diff --git a/Engine/Plugins/Online/OnlineServicesEOSGS/Source/Private/Online/EOSAuthLoginOptionsCommon.h b/Engine/Plugins/Online/OnlineServicesEOSGS/Source/Private/Online/EOSAuthLoginOptionsCommon.h
--- a/Engine/Plugins/Online/OnlineServicesEOSGS/Source/Private/Online/EOSAuthLoginOptionsCommon.h
+++ b/Engine/Plugins/Online/OnlineServicesEOSGS/Source/Private/Online/EOSAuthLoginOptionsCommon.h
@@ -24,13 +24,14 @@ ENUM_CLASS_FLAGS(EEOSAuthTranslationFlags);
 
 struct FEOSAuthTranslationTraits
 {
-	EOS_ELoginCredentialType Type;
-	EEOSAuthTranslationFlags Flags;
+	EOS_ELoginCredentialType Type = EOS_ELoginCredentialType::EOS_LCT_Password;
+	EEOSAuthTranslationFlags Flags = EEOSAuthTranslationFlags::None;
 };
 
 struct FEOSExternalAuthTranslationTraits
 {
-	EOS_EExternalCredentialType Type;
+	EOS_EExternalCredentialType Type = EOS_EExternalCredentialType::EOS_ECT_EPIC;
+	EOS_ELinkAccountFlags LinkAccountFlags = EOS_ELinkAccountFlags::EOS_LA_NoFlags;
 };
 
 class FEOSAuthLoginOptionsCommon : public EOS_Auth_LoginOptions
@@ -41,6 +42,8 @@ public:
 	FEOSAuthLoginOptionsCommon(FEOSAuthLoginOptionsCommon&&);
 	FEOSAuthLoginOptionsCommon& operator=(FEOSAuthLoginOptionsCommon&&);
 
+	EOS_ELinkAccountFlags GetLinkAccountFlags() const { return LinkAccountFlags; }
+
 protected:
 	FEOSAuthLoginOptionsCommon();
 
@@ -50,6 +53,7 @@ protected:
 	static bool InitSystemAuthCredentialOptions(FEOSAuthLoginOptionsCommon& Options) { return true; }
 
 	EOS_Auth_Credentials CredentialsData;
+	EOS_ELinkAccountFlags LinkAccountFlags = EOS_ELinkAccountFlags::EOS_LA_NoFlags;
 	TArray<char> IdUtf8;
 	TArray<char> TokenUtf8;
 };
diff --git a/Engine/Plugins/Online/OnlineServicesEOSGS/Source/Public/Online/AuthEOSGS.h b/Engine/Plugins/Online/OnlineServicesEOSGS/Source/Public/Online/AuthEOSGS.h
--- a/Engine/Plugins/Online/OnlineServicesEOSGS/Source/Public/Online/AuthEOSGS.h
+++ b/Engine/Plugins/Online/OnlineServicesEOSGS/Source/Public/Online/AuthEOSGS.h
@@ -33,6 +33,7 @@ struct FAuthLoginEASImpl
 		FString CredentialsId;
 		TVariant<FString, FExternalAuthToken> CredentialsToken;
 		TArray<FString> Scopes;
+		bool bAutoLinkAccount = true;
 	};
 
 	struct Result
@@ -205,7 +206,7 @@ public:
 	virtual TOnlineAsyncOpHandle<FAuthCancelVerifiedAuthTicket> CancelVerifiedAuthTicket(FAuthCancelVerifiedAuthTicket::Params&& Params) override;
 	virtual TOnlineAsyncOpHandle<FAuthBeginVerifiedAuthSession> BeginVerifiedAuthSession(FAuthBeginVerifiedAuthSession::Params&& Params) override;
 	virtual TOnlineAsyncOpHandle<FAuthEndVerifiedAuthSession> EndVerifiedAuthSession(FAuthEndVerifiedAuthSession::Params&& Params) override;
-	virtual TOnlineResult<FAuthGetRelyingParty> GetRelyingParty() const override;
+	virtual TOnlineResult<FAuthGetRelyingParty> GetRelyingParty(FAuthGetRelyingParty::Params&& Params) const override;
 	// End IAuth
 
 	// Begin FAuthEOSGS
@@ -236,6 +237,24 @@ protected:
 	void OnEASLoginStatusChanged(const EOS_Auth_LoginStatusChangedCallbackInfo* Data);
 
 protected:
+	template <typename ValueType>
+	using TLocalUserArray = TSparseArray<ValueType, TInlineSparseArrayAllocator<MAX_LOCAL_PLAYERS>>;
+
+	struct FLoginContinuationData
+	{
+		FLoginContinuationId ContinuationId;
+		EOS_ContinuanceToken ContinuanceToken;
+		EOS_ELinkAccountFlags LinkAccountFlags = EOS_ELinkAccountFlags::EOS_LA_NoFlags;
+	};
+
+	struct FUserScopedData
+	{
+		FLoginContinuationId LastLoginContinuationId;
+		TArray<FLoginContinuationData> LoginContinuations;
+	};
+
+	uint32 NextLoginContinuationId = 1;
+
 #if !UE_BUILD_SHIPPING
 	static void CheckMetadata();
 #endif
@@ -246,12 +265,17 @@ protected:
 
 	static FAccountId CreateAccountId(const EOS_ProductUserId ProductUserId);
 
+	FUserScopedData* GetUserScopedData(FPlatformUserId PlatformUserId);
+	const FUserScopedData* GetUserScopedData(FPlatformUserId PlatformUserId) const;
+	FUserScopedData* GetOrCreateUserScopedData(FPlatformUserId PlatformUserId);
+
 	EOS_HAuth AuthHandle = nullptr;
 	EOS_HConnect ConnectHandle = nullptr;
 	EOSEventRegistrationPtr OnConnectLoginStatusChangedEOSEventRegistration;
 	EOSEventRegistrationPtr OnConnectAuthNotifyExpirationEOSEventRegistration;
 	EOSEventRegistrationPtr OnAuthLoginStatusChangedEOSEventRegistration;
 	FAccountInfoRegistryEOS AccountInfoRegistryEOS;
+	TLocalUserArray<FUserScopedData> UserScopedData;
 };
 
 namespace Meta {
@@ -261,7 +285,8 @@ BEGIN_ONLINE_STRUCT_META(FAuthLoginEASImpl::Params)
 	ONLINE_STRUCT_FIELD(FAuthLoginEASImpl::Params, CredentialsType),
 	ONLINE_STRUCT_FIELD(FAuthLoginEASImpl::Params, CredentialsId),
 	ONLINE_STRUCT_FIELD(FAuthLoginEASImpl::Params, CredentialsToken),
-	ONLINE_STRUCT_FIELD(FAuthLoginEASImpl::Params, Scopes)
+	ONLINE_STRUCT_FIELD(FAuthLoginEASImpl::Params, Scopes),
+	ONLINE_STRUCT_FIELD(FAuthLoginEASImpl::Params, bAutoLinkAccount)
 END_ONLINE_STRUCT_META()
 
 BEGIN_ONLINE_STRUCT_META(FAuthLoginEASImpl::Result)
diff --git a/Engine/Source/Runtime/CoreOnline/Private/Online/CoreOnline.cpp b/Engine/Source/Runtime/CoreOnline/Private/Online/CoreOnline.cpp
--- a/Engine/Source/Runtime/CoreOnline/Private/Online/CoreOnline.cpp
+++ b/Engine/Source/Runtime/CoreOnline/Private/Online/CoreOnline.cpp
@@ -465,6 +465,7 @@ FString ToLogString(const FOnlineSessionId& Id) { return ToLogStringImpl(Id); }
 FString ToLogString(const FSessionInviteId& Id) { return ToLogStringImpl(Id); }
 FString ToLogString(const FVerifiedAuthTicketId& Id) { return ToLogStringImpl(Id); }
 FString ToLogString(const FVerifiedAuthSessionId& Id) { return ToLogStringImpl(Id); }
+FString ToLogString(const FLoginContinuationId& Id) { return ToLogStringImpl(Id); }
 
 }	/* UE::Online */
 
diff --git a/Engine/Source/Runtime/CoreOnline/Public/Online/CoreOnline.h b/Engine/Source/Runtime/CoreOnline/Public/Online/CoreOnline.h
--- a/Engine/Source/Runtime/CoreOnline/Public/Online/CoreOnline.h
+++ b/Engine/Source/Runtime/CoreOnline/Public/Online/CoreOnline.h
@@ -221,6 +221,7 @@ namespace OnlineIdHandleTags
 	struct FLobby {};
 	struct FVerifiedAuthTicket {};
 	struct FVerifiedAuthSession {};
+	struct FLoginContinuation {};
 }
 
 enum class EOnlineServices : uint8
@@ -319,6 +320,7 @@ using FOnlineSessionId = TOnlineId<OnlineIdHandleTags::FSession>;
 using FSessionInviteId = TOnlineId<OnlineIdHandleTags::FSessionInvite>;
 using FVerifiedAuthTicketId = TOnlineId<OnlineIdHandleTags::FVerifiedAuthTicket>;
 using FVerifiedAuthSessionId = TOnlineId<OnlineIdHandleTags::FVerifiedAuthSession>;
+using FLoginContinuationId = TOnlineId<OnlineIdHandleTags::FLoginContinuation>;
 
 COREONLINE_API FString ToString(const FAccountId& Id);
 COREONLINE_API FString ToString(const FOnlineSessionId& Id);
@@ -334,6 +336,7 @@ COREONLINE_API FString ToLogString(const FOnlineSessionId& Id);
 COREONLINE_API FString ToLogString(const FSessionInviteId& Id);
 COREONLINE_API FString ToLogString(const FVerifiedAuthTicketId& Id);
 COREONLINE_API FString ToLogString(const FVerifiedAuthSessionId& Id);
+COREONLINE_API FString ToLogString(const FLoginContinuationId& Id);
 
 template<typename IdType>
 inline uint32 GetTypeHash(const TOnlineId<IdType>& OnlineId)
