PR #11703: Fixed the reference of Entitlement

diff --git a/Engine/Plugins/Online/OnlineServicesEOS/Source/Private/Online/CommerceEOS.cpp b/Engine/Plugins/Online/OnlineServicesEOS/Source/Private/Online/CommerceEOS.cpp
--- a/Engine/Plugins/Online/OnlineServicesEOS/Source/Private/Online/CommerceEOS.cpp
+++ b/Engine/Plugins/Online/OnlineServicesEOS/Source/Private/Online/CommerceEOS.cpp
@@ -383,7 +383,7 @@ TOnlineAsyncOpHandle<FCommerceQueryEntitlements> FCommerceEOS::QueryEntitlements
 				}
 
 				// Parse the entitlement into the receipt format
-				FEntitlement Entitlement = Entitlements.Emplace_GetRef();
+				FEntitlement& Entitlement = Entitlements.Emplace_GetRef();
 				EOSEntitlementToOssEntitlement(Entitlement, EosEntitlement);
 
 				EOS_Ecom_Entitlement_Release(EosEntitlement);
