Add InternalAccount LinkAccount tag to handle linking a platform account which cannot be logged into externally.

diff --git a/Engine/Plugins/Online/OnlineServicesEOS/Source/Private/Online/AuthEOS.cpp b/Engine/Plugins/Online/OnlineServicesEOS/Source/Private/Online/AuthEOS.cpp
--- a/Engine/Plugins/Online/OnlineServicesEOS/Source/Private/Online/AuthEOS.cpp
+++ b/Engine/Plugins/Online/OnlineServicesEOS/Source/Private/Online/AuthEOS.cpp
@@ -45,6 +45,11 @@ namespace
 
 /* anonymous */ }
 
+namespace ELinkAccountTag
+{
+const FName InternalAccount = TEXT("InternalAccount");
+}
+
 FAuthEOS::FAuthEOS(FOnlineServicesEOS& InServices)
 	: Super(InServices)
 {
@@ -315,7 +320,7 @@ TOnlineAsyncOpHandle<FAuthLinkAccount> FAuthEOS::LinkAccount(FAuthLinkAccount::P
 		EOS_Auth_LinkAccountOptions LinkAccountOptions = {};
 		LinkAccountOptions.ApiVersion = 1;
 		LinkAccountOptions.ContinuanceToken = LoginContinuationData.ContinuanceToken;
-		LinkAccountOptions.LinkAccountFlags = LoginContinuationData.LinkAccountFlags;
+		LinkAccountOptions.LinkAccountFlags = (!Params.Tags.Contains(ELinkAccountTag::InternalAccount)) ? LoginContinuationData.LinkAccountFlags : EOS_ELinkAccountFlags::EOS_LA_NoFlags;
 		LinkAccountOptions.LocalUserId = AccountInfoEOS ? AccountInfoEOS->EpicAccountId : nullptr;
 		UE_EOS_CHECK_API_MISMATCH(EOS_AUTH_LINKACCOUNT_API_LATEST, 1);
 
@@ -430,7 +435,7 @@ TOnlineAsyncOpHandle<FAuthLinkAccount> FAuthEOS::LinkAccount(FAuthLinkAccount::P
 				{
 					if (LoginResult.IsError())
 					{
-						UE_LOG(LogOnlineServices, Warning, TEXT("[FAuthEOS::Login] Failure: LoginConnectImpl %s"), *LoginResult.GetErrorValue().GetLogString());
+						UE_LOG(LogOnlineServices, Warning, TEXT("[FAuthEOS::LinkAccount] Failure: LoginConnectImpl %s"), *LoginResult.GetErrorValue().GetLogString());
 
 						LogoutEASImpl(FAuthLogoutEASImpl::Params{ AccountInfoEOS->EpicAccountId })
 						.Next([AsyncOp, Error = MoveTemp(LoginResult.GetErrorValue()), Promise = MoveTemp(Promise)](TDefaultErrorResult<FAuthLogoutEASImpl>&&) mutable -> void
@@ -497,7 +502,7 @@ TOnlineAsyncOpHandle<FAuthLinkAccount> FAuthEOS::LinkAccount(FAuthLinkAccount::P
 			else
 			{
 				FOnlineError CopyUserInfoError(Errors::FromEOSResult(CopyBestDisplayNameResult));
-				UE_LOG(LogOnlineServices, Warning, TEXT("[FAuthEOS::Login] Failure: EOS_UserInfo_CopyBestDisplayName %s"), *CopyUserInfoError.GetLogString());
+				UE_LOG(LogOnlineServices, Warning, TEXT("[FAuthEOS::LinkAccount] Failure: EOS_UserInfo_CopyBestDisplayName %s"), *CopyUserInfoError.GetLogString());
 
 				TPromise<void> Promise;
 				TFuture<void> Future = Promise.GetFuture();
@@ -522,14 +527,14 @@ TOnlineAsyncOpHandle<FAuthLinkAccount> FAuthEOS::LinkAccount(FAuthLinkAccount::P
 
 		if (bUserWasLoggedIn)
 		{
-			UE_LOG(LogOnlineServices, Log, TEXT("[FAuthEOS::FAuthLinkAccount] Successfully linked account. AccountId: %s"), *ToLogString(AccountInfoEOS->AccountId));
+			UE_LOG(LogOnlineServices, Log, TEXT("[FAuthEOS::LinkAccount] Successfully linked account. AccountId: %s"), *ToLogString(AccountInfoEOS->AccountId));
 		}
 		else
 		{
 			AccountInfoEOS->LoginStatus = ELoginStatus::LoggedIn;
 			AccountInfoEOS->AccountId = CreateAccountId(AccountInfoEOS->EpicAccountId, AccountInfoEOS->ProductUserId);
 			AccountInfoRegistryEOS.Register(AccountInfoEOS);
-			UE_LOG(LogOnlineServices, Log, TEXT("[FAuthEOS::FAuthLinkAccount] Successfully logged in. AccountId: %s"), *ToLogString(AccountInfoEOS->AccountId));
+			UE_LOG(LogOnlineServices, Log, TEXT("[FAuthEOS::LinkAccount] Successfully logged in. AccountId: %s"), *ToLogString(AccountInfoEOS->AccountId));
 			OnAuthLoginStatusChangedEvent.Broadcast(FAuthLoginStatusChanged{ AccountInfoEOS, AccountInfoEOS->LoginStatus });
 		}
 
diff --git a/Engine/Plugins/Online/OnlineServicesEOS/Source/Public/Online/AuthEOS.h b/Engine/Plugins/Online/OnlineServicesEOS/Source/Public/Online/AuthEOS.h
--- a/Engine/Plugins/Online/OnlineServicesEOS/Source/Public/Online/AuthEOS.h
+++ b/Engine/Plugins/Online/OnlineServicesEOS/Source/Public/Online/AuthEOS.h
@@ -10,6 +10,12 @@ namespace UE::Online {
 
 class FOnlineServicesEOS;
 
+namespace ELinkAccountTag
+{
+// An internal account is an account which has nothing external allowing the user to login on the epicgames.com website.
+ONLINESERVICESEOS_API extern const FName InternalAccount;
+}
+
 class ONLINESERVICESEOS_API FAuthEOS : public FAuthEOSGS
 {
 public:
