- Fix FSocialEOS crash when a friend notification comes in, but we're not logged in! (logged in via a different plugin)

diff --git a/Engine/Plugins/Online/OnlineServicesEOS/Source/Private/Online/SocialEOS.cpp b/Engine/Plugins/Online/OnlineServicesEOS/Source/Private/Online/SocialEOS.cpp
--- a/Engine/Plugins/Online/OnlineServicesEOS/Source/Private/Online/SocialEOS.cpp
+++ b/Engine/Plugins/Online/OnlineServicesEOS/Source/Private/Online/SocialEOS.cpp
@@ -48,12 +48,20 @@ void FSocialEOS::Initialize()
 	{
 		FSocialEOS* This = reinterpret_cast<FSocialEOS*>(Data->ClientData);
 
-		const FAccountId LocalAccountId = FindAccountIdChecked(Data->LocalUserId);
-		This->Services.Get<FAuthEOS>()->ResolveAccountId(LocalAccountId, Data->TargetUserId)
-		.Next([This, LocalAccountId, PreviousStatus = Data->PreviousStatus, CurrentStatus = Data->CurrentStatus](const FAccountId& FriendAccountId)
+		// Note this local user may not exist if login has not been called... (more accurately if login has occured outside of this plugin)
+		const FAccountId LocalAccountId = FindAccountId(Data->LocalUserId);
+		if (LocalAccountId.IsValid())
 		{
-			This->OnEOSFriendsUpdate(LocalAccountId, FriendAccountId, PreviousStatus, CurrentStatus);
-		});
+			This->Services.Get<FAuthEOS>()->ResolveAccountId(LocalAccountId, Data->TargetUserId)
+			.Next([This, LocalAccountId, PreviousStatus = Data->PreviousStatus, CurrentStatus = Data->CurrentStatus](const FAccountId& FriendAccountId)
+			{
+				This->OnEOSFriendsUpdate(LocalAccountId, FriendAccountId, PreviousStatus, CurrentStatus);
+			});
+		}
+		else
+		{
+			UE_LOG(LogOnlineServices, Log, TEXT("EOS_Friends_OnFriendsUpdateInfo LocalUserId=%s not found."), *LexToString(Data->LocalUserId));
+		}
 	});
 }
 
