Added user authentication validations to FCommerceEOS interfaces. Modified the parameter of FCommerceEOS::QueryOffers to accept SandboxId.

diff --git a/Engine/Plugins/Online/OnlineServices/Source/OnlineServicesInterface/Public/Online/Commerce.h b/Engine/Plugins/Online/OnlineServices/Source/OnlineServicesInterface/Public/Online/Commerce.h
--- a/Engine/Plugins/Online/OnlineServices/Source/OnlineServicesInterface/Public/Online/Commerce.h
+++ b/Engine/Plugins/Online/OnlineServices/Source/OnlineServicesInterface/Public/Online/Commerce.h
@@ -72,6 +72,8 @@ struct FCommerceQueryOffers
 	{
 		/* Local user */
 		FAccountId LocalAccountId;
+		/* Optional catalog namespace to use for this query. Can be left blank to use the "default" namespace */
+		FString OverrideCatalogNamespace;
 	};
 
 	struct Result
@@ -316,6 +318,8 @@ public:
 
 class ICommerce
 {
+public:
+
 	/* Query the list of all offers for the target user */
 	virtual TOnlineAsyncOpHandle<FCommerceQueryOffers> QueryOffers(FCommerceQueryOffers::Params&& Params) = 0;
 
diff --git a/Engine/Plugins/Online/OnlineServicesEOS/Source/Private/Online/CommerceEOS.cpp b/Engine/Plugins/Online/OnlineServicesEOS/Source/Private/Online/CommerceEOS.cpp
--- a/Engine/Plugins/Online/OnlineServicesEOS/Source/Private/Online/CommerceEOS.cpp
+++ b/Engine/Plugins/Online/OnlineServicesEOS/Source/Private/Online/CommerceEOS.cpp
@@ -77,6 +77,10 @@ TOnlineAsyncOpHandle<FCommerceQueryOffers> FCommerceEOS::QueryOffers(FCommerceQu
 			const FCommerceQueryOffers::Params& Params = Op.GetParams();
 			EOS_Ecom_QueryOffersOptions Options = { };
 			Options.ApiVersion = 1;
+			FTCHARToUTF8 Utf8OverrideCatalogNamespace(*Params.OverrideCatalogNamespace);
+			if (!Params.OverrideCatalogNamespace.IsEmpty()) {
+				Options.OverrideCatalogNamespace = Utf8OverrideCatalogNamespace.Get();
+			}
 			UE_EOS_CHECK_API_MISMATCH(EOS_ECOM_QUERYOFFERS_API_LATEST, 1);
 			Options.LocalUserId = GetEpicAccountIdChecked(Params.LocalAccountId);
 			EOS_Async(EOS_Ecom_QueryOffers, EcomHandle, Options, MoveTemp(Promise));
@@ -150,6 +154,11 @@ TOnlineAsyncOpHandle<FCommerceQueryOffersById> FCommerceEOS::QueryOffersById(FCo
 
 TOnlineResult<FCommerceGetOffers> FCommerceEOS::GetOffers(FCommerceGetOffers::Params&& Params)
 {
+	if (!Services.Get<FAuthEOS>()->IsLoggedIn(Params.LocalAccountId))
+	{
+		return TOnlineResult<FCommerceGetOffers>(Errors::NotLoggedIn());
+	}
+
 	if(CachedOffers.Contains(Params.LocalAccountId))
 	{
 		return TOnlineResult<FCommerceGetOffers>({CachedOffers.FindChecked(Params.LocalAccountId)});
@@ -159,6 +168,11 @@ TOnlineResult<FCommerceGetOffers> FCommerceEOS::GetOffers(FCommerceGetOffers::Pa
 
 TOnlineResult<FCommerceGetOffersById> FCommerceEOS::GetOffersById(FCommerceGetOffersById::Params&& Params)
 {
+	if (!Services.Get<FAuthEOS>()->IsLoggedIn(Params.LocalAccountId))
+	{
+		return TOnlineResult<FCommerceGetOffersById>(Errors::NotLoggedIn());
+	}
+
 	if (CachedOffers.Contains(Params.LocalAccountId))
 	{
 		return TOnlineResult<FCommerceGetOffersById>({ CachedOffers.FindChecked(Params.LocalAccountId).FilterByPredicate(
@@ -189,12 +203,14 @@ TOnlineAsyncOpHandle<FCommerceCheckout> FCommerceEOS::Checkout(FCommerceCheckout
 		if (!Services.Get<FAuthEOS>()->IsLoggedIn(Params.LocalAccountId))
 		{
 			Op.SetError(Errors::NotLoggedIn());
+			Promise.SetValue(nullptr);
 			return;
 		}
 		EOS_EpicAccountId LocalUserEasId = GetEpicAccountId(Params.LocalAccountId);
 		if (!EOS_EpicAccountId_IsValid(LocalUserEasId))
 		{
 			Op.SetError(Errors::NotLoggedIn());
+			Promise.SetValue(nullptr);
 			return;
 		}
 
@@ -399,6 +415,11 @@ TOnlineAsyncOpHandle<FCommerceQueryEntitlements> FCommerceEOS::QueryEntitlements
 
 TOnlineResult<FCommerceGetEntitlements> FCommerceEOS::GetEntitlements(FCommerceGetEntitlements::Params&& Params)
 {
+	if (!Services.Get<FAuthEOS>()->IsLoggedIn(Params.LocalAccountId))
+	{
+		return TOnlineResult<FCommerceGetEntitlements>(Errors::NotLoggedIn());
+	}
+
 	if (!CachedEntitlements.Contains(Params.LocalAccountId))
 	{
 		return TOnlineResult<FCommerceGetEntitlements>(Errors::NotFound());
