* Fixes a deadlock during program shutdown
* Fixes an issue where low level SDK data was not properly populated in some situations.
* Fixed an issue where a failed API call to get user display names could lead to a crash.

diff --git a/Engine/Plugins/Online/EOSShared/Source/EOSShared/Private/EOSSDKManager.cpp b/Engine/Plugins/Online/EOSShared/Source/EOSShared/Private/EOSSDKManager.cpp
--- a/Engine/Plugins/Online/EOSShared/Source/EOSShared/Private/EOSSDKManager.cpp
+++ b/Engine/Plugins/Online/EOSShared/Source/EOSShared/Private/EOSSDKManager.cpp
@@ -1085,13 +1085,16 @@ void FEOSSDKManager::Shutdown()
 
 		if (ActivePlatforms.Num() > 0)
 		{
-			FRWScopeLock ScopeLock(ActivePlatformsCS, SLT_Write);
+			{
+				FRWScopeLock ScopeLock(ActivePlatformsCS, SLT_Write);
 
-			UE_LOG(LogEOSSDK, Warning, TEXT("FEOSSDKManager::Shutdown Releasing %d remaining platforms"), ActivePlatforms.Num());
+				UE_LOG(LogEOSSDK, Warning, TEXT("FEOSSDKManager::Shutdown Releasing %d remaining platforms"), ActivePlatforms.Num());
+
+				TArray<EOS_HPlatform> ActivePlatformHandles;
+				ActivePlatforms.GenerateKeyArray(ActivePlatformHandles);
+				ReleasedPlatforms.Append(ActivePlatformHandles);
+			}
 
-			TArray<EOS_HPlatform> ActivePlatformHandles;
-			ActivePlatforms.GenerateKeyArray(ActivePlatformHandles);
-			ReleasedPlatforms.Append(ActivePlatformHandles);
 			ReleaseReleasedPlatforms();
 		}
 
diff --git a/Engine/Plugins/Online/OnlineServicesEOS/Source/Private/Online/UserInfoEOS.cpp b/Engine/Plugins/Online/OnlineServicesEOS/Source/Private/Online/UserInfoEOS.cpp
--- a/Engine/Plugins/Online/OnlineServicesEOS/Source/Private/Online/UserInfoEOS.cpp
+++ b/Engine/Plugins/Online/OnlineServicesEOS/Source/Private/Online/UserInfoEOS.cpp
@@ -39,6 +39,8 @@ TOnlineAsyncOpHandle<FQueryUserInfo> FUserInfoEOS::QueryUserInfo(FQueryUserInfo:
 			return Op->GetHandle();
 		}
 
+		const bool bIsLoggedInConnect = EOS_ProductUserId_IsValid(GetProductUserId(Params.LocalAccountId)) == EOS_TRUE;
+
 		for (const FAccountId TargetAccountId : Params.AccountIds)
 		{
 			Op->Then([this, TargetAccountId](TOnlineAsyncOp<FQueryUserInfo>& Op, TPromise<const EOS_UserInfo_QueryUserInfoCallbackInfo*>&& Promise)
@@ -76,6 +78,35 @@ TOnlineAsyncOpHandle<FQueryUserInfo> FUserInfoEOS::QueryUserInfo(FQueryUserInfo:
 					Op.SetError(Errors::FromEOSResult(CallbackInfo->ResultCode));
 				}
 			});
+
+			if (bIsLoggedInConnect)
+			{
+				Op->Then([this, TargetAccountId](TOnlineAsyncOp<FQueryUserInfo>& Op, TPromise<const EOS_Connect_QueryExternalAccountMappingsCallbackInfo*> && Promise)
+				{
+					const FQueryUserInfo::Params& Params = Op.GetParams();
+					const EOS_EpicAccountId TargetUserEasId = GetEpicAccountIdChecked(TargetAccountId);
+					const FString TargetUserEasIdStr = LexToString(TargetUserEasId);
+					const auto TargetUserEasIdUtf8StringCast = StringCast<UTF8CHAR>(*TargetUserEasIdStr);
+					const char* TargetUserEasIdUtf8Str = (const char*)TargetUserEasIdUtf8StringCast.Get();
+					EOS_Connect_QueryExternalAccountMappingsOptions Options;
+					Options.ApiVersion = 1;
+					UE_EOS_CHECK_API_MISMATCH(EOS_CONNECT_QUERYEXTERNALACCOUNTMAPPINGS_API_LATEST, 1);
+					Options.LocalUserId = GetProductUserIdChecked(Params.LocalAccountId);
+					Options.AccountIdType = EOS_EExternalAccountType::EOS_EAT_EPIC;
+					Options.ExternalAccountIds = &TargetUserEasIdUtf8Str;
+					Options.ExternalAccountIdCount = 1;
+					EOS_HConnect ConnectHandle = EOS_Platform_GetConnectInterface( * static_cast<FOnlineServicesEOS&>(GetServices()).GetEOSPlatformHandle());
+					EOS_Async(EOS_Connect_QueryExternalAccountMappings, ConnectHandle, Options, MoveTemp(Promise));
+				})
+				.Then([this](TOnlineAsyncOp<FQueryUserInfo>& Op, const EOS_Connect_QueryExternalAccountMappingsCallbackInfo* CallbackInfo) mutable
+				{
+					if (CallbackInfo->ResultCode != EOS_EResult::EOS_Success)
+					{
+						UE_LOG(LogOnlineServices, Warning, TEXT("EOS_Connect_QueryExternalAccountMappings failed with result=[%s]"), *LexToString(CallbackInfo->ResultCode));
+						Op.SetError(Errors::FromEOSResult(CallbackInfo->ResultCode));
+					}
+				});
+			}
 		}
 
 		Op->Then([](TOnlineAsyncOp<FQueryUserInfo>& Op)
@@ -107,10 +138,13 @@ TOnlineResult<FGetUserInfo> FUserInfoEOS::GetUserInfo(FGetUserInfo::Params&& Par
 	Options.LocalUserId = GetEpicAccountIdChecked(Params.LocalAccountId);
 	Options.TargetUserId = TargetUserEasId;
 
-	EOS_UserInfo_BestDisplayName* EosBestDisplayName;
+	EOS_UserInfo_BestDisplayName* EosBestDisplayName = nullptr;
 	ON_SCOPE_EXIT
 	{
-		EOS_UserInfo_BestDisplayName_Release(EosBestDisplayName);
+		if (EosBestDisplayName)
+		{
+			EOS_UserInfo_BestDisplayName_Release(EosBestDisplayName);
+		}
 	};
 
 	EOS_EResult EosResult = EOS_UserInfo_CopyBestDisplayName(UserInfoHandle, &Options, &EosBestDisplayName);
