- Fix a few async operations capturing "this" and not guarding against it going away
- Add an operator bool for TOnlineId

diff --git a/Engine/Plugins/Online/OnlineServicesEOS/Source/Private/Online/PresenceEOS.cpp b/Engine/Plugins/Online/OnlineServicesEOS/Source/Private/Online/PresenceEOS.cpp
--- a/Engine/Plugins/Online/OnlineServicesEOS/Source/Private/Online/PresenceEOS.cpp
+++ b/Engine/Plugins/Online/OnlineServicesEOS/Source/Private/Online/PresenceEOS.cpp
@@ -74,38 +74,45 @@ void FPresenceEOS::Initialize()
 	check(PresenceHandle != nullptr);
 
 	// Register for friend updates
-	EOS_Presence_AddNotifyOnPresenceChangedOptions Options = { };
-	Options.ApiVersion = 1;
-	UE_EOS_CHECK_API_MISMATCH(EOS_PRESENCE_ADDNOTIFYONPRESENCECHANGED_API_LATEST, 1);
-	NotifyPresenceChangedNotificationId = EOS_Presence_AddNotifyOnPresenceChanged(PresenceHandle, &Options, this, [](const EOS_Presence_PresenceChangedCallbackInfo* Data)
+	OnPresenceChangedCallback = MakeUnique<FOnPresenceChangedCallback>(AsWeak());
+	OnPresenceChangedCallback->CallbackLambda = [this](const EOS_Presence_PresenceChangedCallbackInfo* Data)
 	{
-		FPresenceEOS* This = reinterpret_cast<FPresenceEOS*>(Data->ClientData);
-		const FAccountId LocalAccountId = FindAccountId(Data->LocalUserId);
-
-		if (LocalAccountId.IsValid())
+		if (const FAccountId LocalAccountId = FindAccountId(Data->LocalUserId))
 		{
-			This->Services.Get<FAuthEOS>()->ResolveAccountId(LocalAccountId, Data->PresenceUserId)
-				.Next([This, LocalAccountId](const FAccountId& PresenceAccountId)
+			Services.Get<FAuthEOS>()->ResolveAccountId(LocalAccountId, Data->PresenceUserId)
+			.Next([WeakThis = AsWeak(), LocalAccountId](const FAccountId& PresenceAccountId)
+				{
+					if (TSharedPtr<FPresenceEOS> StrongThis = StaticCastSharedPtr<FPresenceEOS>(WeakThis.Pin()))
 					{
 						UE_LOG(LogOnlineServices, Verbose, TEXT("OnEOSPresenceUpdate: LocalAccountId=[%s] PresenceAccountId=[%s]"), *ToLogString(LocalAccountId), *ToLogString(PresenceAccountId));
-						This->UpdateUserPresence(LocalAccountId, PresenceAccountId);
-					});
+						StrongThis->UpdateUserPresence(LocalAccountId, PresenceAccountId);
+					}
+				});
 		}
 		else // In some cases, this delegate will fire before the login process has completed, so we won't be able to find the AccountId just yet. We'll queue that presence update call for after Login has completed
 		{
 			UE_LOG(LogOnlineServices, Verbose, TEXT("OnEOSPresenceUpdate: Account id not found for Epic id [%s]. Will retry after login completes]"), *LexToString(Data->LocalUserId));
 
-			TArray<EOS_EpicAccountId>& PendingPresenceUpdateArray = This->PendingPresenceUpdates.FindOrAdd(Data->LocalUserId);
+			TArray<EOS_EpicAccountId>& PendingPresenceUpdateArray = PendingPresenceUpdates.FindOrAdd(Data->LocalUserId);
 			PendingPresenceUpdateArray.Add(Data->PresenceUserId);
 		}
-	});
+	};
+
+	EOS_Presence_AddNotifyOnPresenceChangedOptions Options = { };
+	Options.ApiVersion = 1;
+	UE_EOS_CHECK_API_MISMATCH(EOS_PRESENCE_ADDNOTIFYONPRESENCECHANGED_API_LATEST, 1);
+	OnPresenceChangedNotificationId = EOS_Presence_AddNotifyOnPresenceChanged(PresenceHandle, &Options, OnPresenceChangedCallback.Get(), OnPresenceChangedCallback->GetCallbackPtr());
 
 	LoginStatusChangedHandle = Services.Get<IAuth>()->OnLoginStatusChanged().Add(this, &FPresenceEOS::HandleAuthLoginStatusChanged);
 }
 
 void FPresenceEOS::PreShutdown()
 {
-	EOS_Presence_RemoveNotifyOnPresenceChanged(PresenceHandle, NotifyPresenceChangedNotificationId);
+	if (OnPresenceChangedNotificationId != EOS_INVALID_NOTIFICATIONID)
+	{
+		EOS_Presence_RemoveNotifyOnPresenceChanged(PresenceHandle, OnPresenceChangedNotificationId);
+	}
+	OnPresenceChangedCallback.Reset();
 
 	LoginStatusChangedHandle.Unbind();
 }
@@ -123,10 +130,13 @@ void FPresenceEOS::HandleAuthLoginStatusChanged(const FAuthLoginStatusChanged& E
 			for (EOS_EpicAccountId& PresenceAccountId : *PendingPresenceUpdateArray)
 			{
 				Services.Get<FAuthEOS>()->ResolveAccountId(LocalAccountId, PresenceAccountId)
-					.Next([this, LocalAccountId](const FAccountId& PresenceAccountId)
+					.Next([WeakThis = AsWeak(), LocalAccountId](const FAccountId& PresenceAccountId)
 						{
-							UE_LOG(LogOnlineServices, Verbose, TEXT("OnEOSPresenceUpdate: LocalAccountId=[%s] PresenceAccountId=[%s]"), *ToLogString(LocalAccountId), *ToLogString(PresenceAccountId));
-							UpdateUserPresence(LocalAccountId, PresenceAccountId);
+							if (TSharedPtr<FPresenceEOS> StrongThis = StaticCastSharedPtr<FPresenceEOS>(WeakThis.Pin()))
+							{
+								UE_LOG(LogOnlineServices, Verbose, TEXT("OnEOSPresenceUpdate: LocalAccountId=[%s] PresenceAccountId=[%s]"), *ToLogString(LocalAccountId), *ToLogString(PresenceAccountId));
+								StrongThis->UpdateUserPresence(LocalAccountId, PresenceAccountId);
+							}
 						});
 			}
 
diff --git a/Engine/Plugins/Online/OnlineServicesEOS/Source/Public/Online/PresenceEOS.h b/Engine/Plugins/Online/OnlineServicesEOS/Source/Public/Online/PresenceEOS.h
--- a/Engine/Plugins/Online/OnlineServicesEOS/Source/Public/Online/PresenceEOS.h
+++ b/Engine/Plugins/Online/OnlineServicesEOS/Source/Public/Online/PresenceEOS.h
@@ -3,6 +3,7 @@
 #pragma once
 
 #include "CoreMinimal.h"
+#include "EOSSharedTypes.h"
 #include "Online/PresenceCommon.h"
 
 #if defined(EOS_PLATFORM_BASE_FILE_NAME)
@@ -38,12 +39,15 @@ protected:
 protected:
 	EOS_HPresence PresenceHandle = nullptr;
 
+	using FOnPresenceChangedCallback = TEOSGlobalCallback<EOS_Presence_OnPresenceChangedCallback, EOS_Presence_PresenceChangedCallbackInfo, IPresence>;
+	TUniquePtr<FOnPresenceChangedCallback> OnPresenceChangedCallback = nullptr;
+	EOS_NotificationId OnPresenceChangedNotificationId = EOS_INVALID_NOTIFICATIONID;
+
 	/** Login status changed event handle */
 	FOnlineEventDelegateHandle LoginStatusChangedHandle;
 
 	TMap<EOS_EpicAccountId, TArray<EOS_EpicAccountId>> PendingPresenceUpdates;
 	TMap<FAccountId, TMap<FAccountId, TSharedRef<FUserPresence>>> PresenceLists;
-	EOS_NotificationId NotifyPresenceChangedNotificationId = 0;
 };
 
 /* UE::Online */ }
diff --git a/Engine/Source/Runtime/CoreOnline/Public/Online/CoreOnline.h b/Engine/Source/Runtime/CoreOnline/Public/Online/CoreOnline.h
--- a/Engine/Source/Runtime/CoreOnline/Public/Online/CoreOnline.h
+++ b/Engine/Source/Runtime/CoreOnline/Public/Online/CoreOnline.h
@@ -294,6 +294,7 @@ public:
 	}
 
 	inline bool IsValid() const { return GetHandle() != 0; }
+	inline explicit operator bool() const { return IsValid(); }
 
 	EOnlineServices GetOnlineServicesType() const { return EOnlineServices(Value >> 24); }
 	uint32 GetHandle() const { return Value & 0x00FFFFFF; }
