- Add missing F to the start of some EOS callback types (and leave a deprecated typedef behind)
- Migrate some OnlineServicesEOS[GS] components to use the event wrappers rather than manually calling Add/RemoveNotify
- Fix some calls to GetServices using static_cast, but the return type is templated already
- Add some missing super calls

diff --git a/Engine/Plugins/Online/OnlineServicesEOS/Source/Private/Online/AuthEOS.cpp b/Engine/Plugins/Online/OnlineServicesEOS/Source/Private/Online/AuthEOS.cpp
--- a/Engine/Plugins/Online/OnlineServicesEOS/Source/Private/Online/AuthEOS.cpp
+++ b/Engine/Plugins/Online/OnlineServicesEOS/Source/Private/Online/AuthEOS.cpp
@@ -60,7 +60,7 @@ void FAuthEOS::Initialize()
 {
 	Super::Initialize();
 
-	UserInfoHandle = EOS_Platform_GetUserInfoInterface(*static_cast<FOnlineServicesEOS&>(GetServices()).GetEOSPlatformHandle());
+	UserInfoHandle = EOS_Platform_GetUserInfoInterface(*GetServices<FOnlineServicesEOS>().GetEOSPlatformHandle());
 	check(UserInfoHandle != nullptr);
 }
 
diff --git a/Engine/Plugins/Online/OnlineServicesEOS/Source/Private/Online/CommerceEOS.cpp b/Engine/Plugins/Online/OnlineServicesEOS/Source/Private/Online/CommerceEOS.cpp
--- a/Engine/Plugins/Online/OnlineServicesEOS/Source/Private/Online/CommerceEOS.cpp
+++ b/Engine/Plugins/Online/OnlineServicesEOS/Source/Private/Online/CommerceEOS.cpp
@@ -14,7 +14,7 @@ void FCommerceEOS::PostInitialize()
 {	
 	Super::PostInitialize();
 
-	EcomHandle = EOS_Platform_GetEcomInterface(*static_cast<FOnlineServicesEOS&>(GetServices()).GetEOSPlatformHandle());
+	EcomHandle = EOS_Platform_GetEcomInterface(*GetServices<FOnlineServicesEOS>().GetEOSPlatformHandle());
 	check(EcomHandle);
 }
 
diff --git a/Engine/Plugins/Online/OnlineServicesEOS/Source/Private/Online/ExternalUIEOS.cpp b/Engine/Plugins/Online/OnlineServicesEOS/Source/Private/Online/ExternalUIEOS.cpp
--- a/Engine/Plugins/Online/OnlineServicesEOS/Source/Private/Online/ExternalUIEOS.cpp
+++ b/Engine/Plugins/Online/OnlineServicesEOS/Source/Private/Online/ExternalUIEOS.cpp
@@ -21,7 +21,7 @@ void FExternalUIEOS::Initialize()
 {
 	Super::Initialize();
 
-	UIHandle = EOS_Platform_GetUIInterface(*static_cast<FOnlineServicesEOS&>(GetServices()).GetEOSPlatformHandle());
+	UIHandle = EOS_Platform_GetUIInterface(*GetServices<FOnlineServicesEOS>().GetEOSPlatformHandle());
 	check(UIHandle != nullptr);
 }
 
@@ -31,11 +31,6 @@ void FExternalUIEOS::UpdateConfig()
 	LoadConfig(Config);
 }
 
-void FExternalUIEOS::PreShutdown()
-{
-	Super::PreShutdown();
-}
-
 TOnlineAsyncOpHandle<FExternalUIShowLoginUI> FExternalUIEOS::ShowLoginUI(FExternalUIShowLoginUI::Params&& Params)
 {
 	TOnlineAsyncOpRef<FExternalUIShowLoginUI> Op = GetOp<FExternalUIShowLoginUI>(MoveTemp(Params));
diff --git a/Engine/Plugins/Online/OnlineServicesEOS/Source/Private/Online/PresenceEOS.cpp b/Engine/Plugins/Online/OnlineServicesEOS/Source/Private/Online/PresenceEOS.cpp
--- a/Engine/Plugins/Online/OnlineServicesEOS/Source/Private/Online/PresenceEOS.cpp
+++ b/Engine/Plugins/Online/OnlineServicesEOS/Source/Private/Online/PresenceEOS.cpp
@@ -62,59 +62,36 @@ static inline EOS_Presence_EStatus ToEOS_Presence_EStatus(EUserPresenceStatus In
 }
 
 FPresenceEOS::FPresenceEOS(FOnlineServicesEOS& InServices)
-	: FPresenceCommon(InServices)
+	: Super(InServices)
 {
 }
 
 void FPresenceEOS::Initialize()
 {
-	FPresenceCommon::Initialize();
+	Super::Initialize();
 
-	PresenceHandle = EOS_Platform_GetPresenceInterface(*static_cast<FOnlineServicesEOS&>(GetServices()).GetEOSPlatformHandle());
+	PresenceHandle = EOS_Platform_GetPresenceInterface(*GetServices<FOnlineServicesEOS>().GetEOSPlatformHandle());
 	check(PresenceHandle != nullptr);
 
-	// Register for friend updates
-	OnPresenceChangedCallback = MakeUnique<FOnPresenceChangedCallback>(AsWeak());
-	OnPresenceChangedCallback->CallbackLambda = [this](const EOS_Presence_PresenceChangedCallbackInfo* Data)
-	{
-		if (const FAccountId LocalAccountId = FindAccountId(Data->LocalUserId))
-		{
-			Services.Get<FAuthEOS>()->ResolveAccountId(LocalAccountId, Data->PresenceUserId)
-			.Next([WeakThis = AsWeak(), LocalAccountId](const FAccountId& PresenceAccountId)
-				{
-					if (TSharedPtr<FPresenceEOS> StrongThis = StaticCastSharedPtr<FPresenceEOS>(WeakThis.Pin()))
-					{
-						UE_LOG(LogOnlineServices, Verbose, TEXT("OnEOSPresenceUpdate: LocalAccountId=[%s] PresenceAccountId=[%s]"), *ToLogString(LocalAccountId), *ToLogString(PresenceAccountId));
-						StrongThis->UpdateUserPresence(LocalAccountId, PresenceAccountId);
-					}
-				});
-		}
-		else // In some cases, this delegate will fire before the login process has completed, so we won't be able to find the AccountId just yet. We'll queue that presence update call for after Login has completed
-		{
-			UE_LOG(LogOnlineServices, Verbose, TEXT("OnEOSPresenceUpdate: Account id not found for Epic id [%s]. Will retry after login completes]"), *LexToString(Data->LocalUserId));
-
-			TArray<EOS_EpicAccountId>& PendingPresenceUpdateArray = PendingPresenceUpdates.FindOrAdd(Data->LocalUserId);
-			PendingPresenceUpdateArray.Add(Data->PresenceUserId);
-		}
-	};
+	LoginStatusChangedHandle = Services.Get<IAuth>()->OnLoginStatusChanged().Add(this, &FPresenceEOS::HandleAuthLoginStatusChanged);
 
-	EOS_Presence_AddNotifyOnPresenceChangedOptions Options = { };
-	Options.ApiVersion = 1;
+	OnPresenceChanged = EOS_RegisterComponentEventHandler(
+		this,
+		PresenceHandle,
+		1,
+		&EOS_Presence_AddNotifyOnPresenceChanged,
+		&EOS_Presence_RemoveNotifyOnPresenceChanged,
+		&FPresenceEOS::HandlePresenceChanged);
 	UE_EOS_CHECK_API_MISMATCH(EOS_PRESENCE_ADDNOTIFYONPRESENCECHANGED_API_LATEST, 1);
-	OnPresenceChangedNotificationId = EOS_Presence_AddNotifyOnPresenceChanged(PresenceHandle, &Options, OnPresenceChangedCallback.Get(), OnPresenceChangedCallback->GetCallbackPtr());
-
-	LoginStatusChangedHandle = Services.Get<IAuth>()->OnLoginStatusChanged().Add(this, &FPresenceEOS::HandleAuthLoginStatusChanged);
 }
 
 void FPresenceEOS::PreShutdown()
 {
-	if (OnPresenceChangedNotificationId != EOS_INVALID_NOTIFICATIONID)
-	{
-		EOS_Presence_RemoveNotifyOnPresenceChanged(PresenceHandle, OnPresenceChangedNotificationId);
-	}
-	OnPresenceChangedCallback.Reset();
+	Super::PreShutdown();
 
 	LoginStatusChangedHandle.Unbind();
+
+	OnPresenceChanged = nullptr;
 }
 
 void FPresenceEOS::HandleAuthLoginStatusChanged(const FAuthLoginStatusChanged& EventParameters)
@@ -712,4 +689,27 @@ void FPresenceEOS::UpdateUserPresence(FAccountId LocalAccountId, FAccountId Pres
 	}
 }
 
+void FPresenceEOS::HandlePresenceChanged(const EOS_Presence_PresenceChangedCallbackInfo* Data)
+{
+	if (const FAccountId LocalAccountId = FindAccountId(Data->LocalUserId))
+	{
+		Services.Get<FAuthEOS>()->ResolveAccountId(LocalAccountId, Data->PresenceUserId)
+		.Next([WeakThis = AsWeak(), LocalAccountId, Func = __FUNCTION__](const FAccountId& PresenceAccountId)
+			{
+				if (TSharedPtr<FPresenceEOS> StrongThis = StaticCastSharedPtr<FPresenceEOS>(WeakThis.Pin()))
+				{
+					UE_LOG(LogOnlineServices, Verbose, TEXT("%hs LocalAccountId=[%s] PresenceAccountId=[%s]"), Func, *ToLogString(LocalAccountId), *ToLogString(PresenceAccountId));
+					StrongThis->UpdateUserPresence(LocalAccountId, PresenceAccountId);
+				}
+			});
+	}
+	else // In some cases, this delegate will fire before the login process has completed, so we won't be able to find the AccountId just yet. We'll queue that presence update call for after Login has completed
+	{
+		UE_LOG(LogOnlineServices, Verbose, TEXT("%hs Account id not found for Epic id [%s]. Will retry after login completes]"), __FUNCTION__, *LexToString(Data->LocalUserId));
+
+		TArray<EOS_EpicAccountId>& PendingPresenceUpdateArray = PendingPresenceUpdates.FindOrAdd(Data->LocalUserId);
+		PendingPresenceUpdateArray.Add(Data->PresenceUserId);
+	}
+}
+
 /* UE::Online */ }
diff --git a/Engine/Plugins/Online/OnlineServicesEOS/Source/Private/Online/SocialEOS.cpp b/Engine/Plugins/Online/OnlineServicesEOS/Source/Private/Online/SocialEOS.cpp
--- a/Engine/Plugins/Online/OnlineServicesEOS/Source/Private/Online/SocialEOS.cpp
+++ b/Engine/Plugins/Online/OnlineServicesEOS/Source/Private/Online/SocialEOS.cpp
@@ -35,39 +35,26 @@ void FSocialEOS::Initialize()
 {
 	FSocialCommon::Initialize();
 
-	FriendsHandle = EOS_Platform_GetFriendsInterface(*static_cast<FOnlineServicesEOS&>(GetServices()).GetEOSPlatformHandle());
+	FriendsHandle = EOS_Platform_GetFriendsInterface(*GetServices<FOnlineServicesEOS>().GetEOSPlatformHandle());
 	check(FriendsHandle != nullptr);
 
-	UIHandle = EOS_Platform_GetUIInterface(*static_cast<FOnlineServicesEOS&>(GetServices()).GetEOSPlatformHandle());
+	UIHandle = EOS_Platform_GetUIInterface(*GetServices<FOnlineServicesEOS>().GetEOSPlatformHandle());
 
-	// Register for friend updates
-	EOS_Friends_AddNotifyFriendsUpdateOptions Options = { };
-	Options.ApiVersion = 1;
+	OnFriendsUpdate = EOS_RegisterComponentEventHandler(
+		this,
+		FriendsHandle,
+		1,
+		&EOS_Friends_AddNotifyFriendsUpdate,
+		&EOS_Friends_RemoveNotifyFriendsUpdate,
+		&FSocialEOS::HandleFriendsUpdate);
 	UE_EOS_CHECK_API_MISMATCH(EOS_FRIENDS_ADDNOTIFYFRIENDSUPDATE_API_LATEST, 1);
-	NotifyFriendsUpdateNotificationId = EOS_Friends_AddNotifyFriendsUpdate(FriendsHandle, &Options, this, [](const EOS_Friends_OnFriendsUpdateInfo* Data)
-	{
-		FSocialEOS* This = reinterpret_cast<FSocialEOS*>(Data->ClientData);
-
-		// Note this local user may not exist if login has not been called... (more accurately if login has occured outside of this plugin)
-		const FAccountId LocalAccountId = FindAccountId(Data->LocalUserId);
-		if (LocalAccountId.IsValid())
-		{
-			This->Services.Get<FAuthEOS>()->ResolveAccountId(LocalAccountId, Data->TargetUserId)
-			.Next([This, LocalAccountId, PreviousStatus = Data->PreviousStatus, CurrentStatus = Data->CurrentStatus](const FAccountId& FriendAccountId)
-			{
-				This->OnEOSFriendsUpdate(LocalAccountId, FriendAccountId, PreviousStatus, CurrentStatus);
-			});
-		}
-		else
-		{
-			UE_LOG(LogOnlineServices, Log, TEXT("EOS_Friends_OnFriendsUpdateInfo LocalUserId=%s not found."), *LexToString(Data->LocalUserId));
-		}
-	});
 }
 
 void FSocialEOS::PreShutdown()
 {
-	EOS_Friends_RemoveNotifyFriendsUpdate(FriendsHandle, NotifyFriendsUpdateNotificationId);
+	Super::PreShutdown();
+
+	OnFriendsUpdate = nullptr;
 }
 
 TOnlineAsyncOpHandle<FQueryFriends> FSocialEOS::QueryFriends(FQueryFriends::Params&& InParams)
@@ -407,34 +394,46 @@ TOnlineAsyncOpHandle<FBlockUser> FSocialEOS::BlockUser(FBlockUser::Params&& Para
 	return Op->GetHandle();
 }
 
-void FSocialEOS::OnEOSFriendsUpdate(FAccountId LocalAccountId, FAccountId FriendAccountId, EOS_EFriendsStatus PreviousStatus, EOS_EFriendsStatus CurrentStatus)
-{
-	const EOS_EpicAccountId LocalUserEasId = GetEpicAccountIdChecked(LocalAccountId);
-
-	TMap<FAccountId, TSharedRef<FFriend>>& FriendsList = FriendsLists.FindOrAdd(LocalAccountId);
-
-	if (CurrentStatus == EOS_EFriendsStatus::EOS_FS_NotFriends)
+void FSocialEOS::HandleFriendsUpdate(const EOS_Friends_OnFriendsUpdateInfo* Data)
+{		
+	// Note this local user may not exist if login has not been called... (more accurately if login has occured outside of this plugin)
+	const FAccountId LocalAccountId = FindAccountId(Data->LocalUserId);
+	if (!LocalAccountId.IsValid())
 	{
-		FriendsList.Remove(FriendAccountId);
+		UE_LOG(LogOnlineServices, Log, TEXT("%hs LocalUserId=%s not found."), __FUNCTION__, *LexToString(Data->LocalUserId));
+		return;
 	}
-	else
+
+	Services.Get<FAuthEOS>()->ResolveAccountId(LocalAccountId, Data->TargetUserId)
+	.Next([this, LocalAccountId, PreviousStatus = Data->PreviousStatus, CurrentStatus = Data->CurrentStatus](const FAccountId& FriendAccountId)
 	{
-		// Add/update friend with appropriate status
-		if (TSharedRef<FFriend>* ExistingFriend = FriendsList.Find(FriendAccountId))
+		const EOS_EpicAccountId LocalUserEasId = GetEpicAccountIdChecked(LocalAccountId);
+
+		TMap<FAccountId, TSharedRef<FFriend>>& FriendsList = FriendsLists.FindOrAdd(LocalAccountId);
+
+		if (CurrentStatus == EOS_EFriendsStatus::EOS_FS_NotFriends)
 		{
-			(*ExistingFriend)->Relationship = EOSEFriendStatusToERelationship(CurrentStatus);
+			FriendsList.Remove(FriendAccountId);
 		}
 		else
 		{
-			// Create new entry
-			TSharedRef<FFriend> Friend = MakeShared<FFriend>();
-			Friend->FriendId = FriendAccountId;
-			Friend->Relationship = EOSEFriendStatusToERelationship(CurrentStatus);
-			FriendsList.Emplace(FriendAccountId, MoveTemp(Friend));
+			// Add/update friend with appropriate status
+			if (TSharedRef<FFriend>* ExistingFriend = FriendsList.Find(FriendAccountId))
+			{
+				(*ExistingFriend)->Relationship = EOSEFriendStatusToERelationship(CurrentStatus);
+			}
+			else
+			{
+				// Create new entry
+				TSharedRef<FFriend> Friend = MakeShared<FFriend>();
+				Friend->FriendId = FriendAccountId;
+				Friend->Relationship = EOSEFriendStatusToERelationship(CurrentStatus);
+				FriendsList.Emplace(FriendAccountId, MoveTemp(Friend));
+			}
 		}
-	}
 
-	BroadcastRelationshipUpdated(LocalAccountId, FriendAccountId, EOSEFriendStatusToERelationship(PreviousStatus), EOSEFriendStatusToERelationship(CurrentStatus));
+		BroadcastRelationshipUpdated(LocalAccountId, FriendAccountId, EOSEFriendStatusToERelationship(PreviousStatus), EOSEFriendStatusToERelationship(CurrentStatus));
+	});
 }
 
 /* UE::Online */ }
diff --git a/Engine/Plugins/Online/OnlineServicesEOS/Source/Private/Online/UserInfoEOS.cpp b/Engine/Plugins/Online/OnlineServicesEOS/Source/Private/Online/UserInfoEOS.cpp
--- a/Engine/Plugins/Online/OnlineServicesEOS/Source/Private/Online/UserInfoEOS.cpp
+++ b/Engine/Plugins/Online/OnlineServicesEOS/Source/Private/Online/UserInfoEOS.cpp
@@ -22,7 +22,7 @@ void FUserInfoEOS::Initialize()
 {
 	Super::Initialize();
 
-	UserInfoHandle = EOS_Platform_GetUserInfoInterface(*static_cast<FOnlineServicesEOS&>(GetServices()).GetEOSPlatformHandle());
+	UserInfoHandle = EOS_Platform_GetUserInfoInterface(*GetServices<FOnlineServicesEOS>().GetEOSPlatformHandle());
 	check(UserInfoHandle != nullptr);
 }
 
@@ -95,7 +95,7 @@ TOnlineAsyncOpHandle<FQueryUserInfo> FUserInfoEOS::QueryUserInfo(FQueryUserInfo:
 					Options.AccountIdType = EOS_EExternalAccountType::EOS_EAT_EPIC;
 					Options.ExternalAccountIds = &TargetUserEasIdUtf8Str;
 					Options.ExternalAccountIdCount = 1;
-					EOS_HConnect ConnectHandle = EOS_Platform_GetConnectInterface( * static_cast<FOnlineServicesEOS&>(GetServices()).GetEOSPlatformHandle());
+					EOS_HConnect ConnectHandle = EOS_Platform_GetConnectInterface(*GetServices<FOnlineServicesEOS>().GetEOSPlatformHandle());
 					EOS_Async(EOS_Connect_QueryExternalAccountMappings, ConnectHandle, Options, MoveTemp(Promise));
 				})
 				.Then([this](TOnlineAsyncOp<FQueryUserInfo>& Op, const EOS_Connect_QueryExternalAccountMappingsCallbackInfo* CallbackInfo) mutable
diff --git a/Engine/Plugins/Online/OnlineServicesEOS/Source/Public/Online/ExternalUIEOS.h b/Engine/Plugins/Online/OnlineServicesEOS/Source/Public/Online/ExternalUIEOS.h
--- a/Engine/Plugins/Online/OnlineServicesEOS/Source/Public/Online/ExternalUIEOS.h
+++ b/Engine/Plugins/Online/OnlineServicesEOS/Source/Public/Online/ExternalUIEOS.h
@@ -28,7 +28,6 @@ public:
 	FExternalUIEOS(FOnlineServicesEOS& InOwningSubsystem);
 	virtual void Initialize() override;
 	virtual void UpdateConfig() override;
-	virtual void PreShutdown() override;
 	virtual TOnlineAsyncOpHandle<FExternalUIShowLoginUI> ShowLoginUI(FExternalUIShowLoginUI::Params&& Params) override;
 	virtual TOnlineAsyncOpHandle<FExternalUIShowFriendsUI> ShowFriendsUI(FExternalUIShowFriendsUI::Params&& Params) override;
 
diff --git a/Engine/Plugins/Online/OnlineServicesEOS/Source/Public/Online/PresenceEOS.h b/Engine/Plugins/Online/OnlineServicesEOS/Source/Public/Online/PresenceEOS.h
--- a/Engine/Plugins/Online/OnlineServicesEOS/Source/Public/Online/PresenceEOS.h
+++ b/Engine/Plugins/Online/OnlineServicesEOS/Source/Public/Online/PresenceEOS.h
@@ -2,8 +2,8 @@
 
 #pragma once
 
-#include "CoreMinimal.h"
 #include "EOSSharedTypes.h"
+#include "Online/OnlineServicesEOSGSTypes.h"
 #include "Online/PresenceCommon.h"
 
 #if defined(EOS_PLATFORM_BASE_FILE_NAME)
@@ -18,6 +18,8 @@ class FOnlineServicesEOS;
 class FPresenceEOS : public FPresenceCommon
 {
 public:
+	using Super = FPresenceCommon;
+
 	FPresenceEOS(FOnlineServicesEOS& InServices);
 
 	virtual void Initialize() override;
@@ -36,18 +38,17 @@ protected:
 	/** Performs queued presence updates after a user's login completes */
 	void HandleAuthLoginStatusChanged(const FAuthLoginStatusChanged& EventParameters);
 
-protected:
-	EOS_HPresence PresenceHandle = nullptr;
+	void HandlePresenceChanged(const EOS_Presence_PresenceChangedCallbackInfo* Data);
 
-	using FOnPresenceChangedCallback = TEOSGlobalCallback<EOS_Presence_OnPresenceChangedCallback, EOS_Presence_PresenceChangedCallbackInfo, IPresence>;
-	TUniquePtr<FOnPresenceChangedCallback> OnPresenceChangedCallback = nullptr;
-	EOS_NotificationId OnPresenceChangedNotificationId = EOS_INVALID_NOTIFICATIONID;
+	EOS_HPresence PresenceHandle = nullptr;
 
 	/** Login status changed event handle */
 	FOnlineEventDelegateHandle LoginStatusChangedHandle;
 
 	TMap<EOS_EpicAccountId, TArray<EOS_EpicAccountId>> PendingPresenceUpdates;
 	TMap<FAccountId, TMap<FAccountId, TSharedRef<FUserPresence>>> PresenceLists;
+	
+	FEOSEventRegistrationPtr OnPresenceChanged;
 };
 
 /* UE::Online */ }
diff --git a/Engine/Plugins/Online/OnlineServicesEOS/Source/Public/Online/SocialEOS.h b/Engine/Plugins/Online/OnlineServicesEOS/Source/Public/Online/SocialEOS.h
--- a/Engine/Plugins/Online/OnlineServicesEOS/Source/Public/Online/SocialEOS.h
+++ b/Engine/Plugins/Online/OnlineServicesEOS/Source/Public/Online/SocialEOS.h
@@ -3,6 +3,7 @@
 #pragma once
 
 #include "CoreMinimal.h"
+#include "Online/OnlineServicesEOSGSTypes.h"
 #include "Online/SocialCommon.h"
 
 #if defined(EOS_PLATFORM_BASE_FILE_NAME)
@@ -18,6 +19,8 @@ class FOnlineServicesEOS;
 class FSocialEOS : public FSocialCommon
 {
 public:
+	using Super = FSocialCommon;
+
 	FSocialEOS(FOnlineServicesEOS& InServices);
 
 	virtual void Initialize() override;
@@ -31,13 +34,13 @@ public:
 	virtual TOnlineAsyncOpHandle<FBlockUser> BlockUser(FBlockUser::Params&& Params) override;
 
 protected:
-	void OnEOSFriendsUpdate(FAccountId LocalAccountId, FAccountId FriendAccountId, EOS_EFriendsStatus PreviousStatus, EOS_EFriendsStatus CurrentStatus);
+	void HandleFriendsUpdate(const EOS_Friends_OnFriendsUpdateInfo* Data);
 
 	TMap<FAccountId, TMap<FAccountId, TSharedRef<FFriend>>> FriendsLists;
 	EOS_HFriends FriendsHandle = nullptr;
 	EOS_HUI UIHandle = nullptr;
 
-	EOS_NotificationId NotifyFriendsUpdateNotificationId = 0;
+	FEOSEventRegistrationPtr OnFriendsUpdate;
 };
 
 /* UE::Online */ }
diff --git a/Engine/Plugins/Online/OnlineServicesEOSGS/Source/Private/Online/AchievementsEOSGS.cpp b/Engine/Plugins/Online/OnlineServicesEOSGS/Source/Private/Online/AchievementsEOSGS.cpp
--- a/Engine/Plugins/Online/OnlineServicesEOSGS/Source/Private/Online/AchievementsEOSGS.cpp
+++ b/Engine/Plugins/Online/OnlineServicesEOSGS/Source/Private/Online/AchievementsEOSGS.cpp
@@ -23,52 +23,24 @@ void FAchievementsEOSGS::Initialize()
 {
 	Super::Initialize();
 
-	AchievementsHandle = EOS_Platform_GetAchievementsInterface(*static_cast<FOnlineServicesEOSGS&>(GetServices()).GetEOSPlatformHandle());
+	AchievementsHandle = EOS_Platform_GetAchievementsInterface(*GetServices<FOnlineServicesEOSGS>().GetEOSPlatformHandle());
 	check(AchievementsHandle);
 
-	// Register for achievement unlocked events
-	EOS_Achievements_AddNotifyAchievementsUnlockedV2Options Options = { };
-	Options.ApiVersion = 2;
+	OnAchievementsUnlocked = EOS_RegisterComponentEventHandler(
+		this,
+		AchievementsHandle,
+		2,
+		&EOS_Achievements_AddNotifyAchievementsUnlockedV2,
+		&EOS_Achievements_RemoveNotifyAchievementsUnlocked,
+		&FAchievementsEOSGS::HandleAchievementsUnlocked);
 	UE_EOS_CHECK_API_MISMATCH(EOS_ACHIEVEMENTS_ADDNOTIFYACHIEVEMENTSUNLOCKEDV2_API_LATEST, 2);
-	NotifyAchievementsUnlockedNotificationId = EOS_Achievements_AddNotifyAchievementsUnlockedV2(AchievementsHandle, &Options, this, [](const EOS_Achievements_OnAchievementsUnlockedCallbackV2Info* Data)
-	{
-		FAchievementsEOSGS* This = reinterpret_cast<FAchievementsEOSGS*>(Data->ClientData);
-
-		const FAccountId LocalAccountId = FindAccountId(Data->UserId);
-		if (LocalAccountId.IsValid())
-		{
-			const FString AchievementId = UTF8_TO_TCHAR(Data->AchievementId);
-			const FDateTime UnlockTime = FDateTime::FromUnixTimestamp(Data->UnlockTime);
-
-			UE_LOG(LogOnlineServices, Verbose, TEXT("EOS_Achievements_OnAchievementsUnlocked User=%s AchievId=%s"), *ToLogString(LocalAccountId), *AchievementId);
-
-			if (FAchievementStateMap* LocalUserAchievementStates = This->AchievementStates.Find(LocalAccountId))
-			{
-				if (FAchievementState* AchievementState = LocalUserAchievementStates->Find(AchievementId))
-				{
-					AchievementState->Progress = 1.0f;
-					AchievementState->UnlockTime = UnlockTime;
-				}
-			}
-			
-			FAchievementStateUpdated EventParams;
-			EventParams.LocalAccountId = LocalAccountId;
-			EventParams.AchievementIds.Emplace(AchievementId);
-
-			This->OnAchievementStateUpdatedEvent.Broadcast(EventParams);
-		}
-	});
 }
 
 void FAchievementsEOSGS::Shutdown()
 {
 	Super::Shutdown();
 
-	if (NotifyAchievementsUnlockedNotificationId != EOS_INVALID_NOTIFICATIONID)
-	{
-		EOS_Achievements_RemoveNotifyAchievementsUnlocked(AchievementsHandle, NotifyAchievementsUnlockedNotificationId);
-		NotifyAchievementsUnlockedNotificationId = EOS_INVALID_NOTIFICATIONID;
-	}
+	OnAchievementsUnlocked = nullptr;
 }
 
 TOnlineAsyncOpHandle<FQueryAchievementDefinitions> FAchievementsEOSGS::QueryAchievementDefinitions(FQueryAchievementDefinitions::Params&& Params)
@@ -433,4 +405,31 @@ TOnlineAsyncOpHandle<FUnlockAchievements> FAchievementsEOSGS::UnlockAchievements
 	return Op->GetHandle();
 }
 
+void FAchievementsEOSGS::HandleAchievementsUnlocked(const EOS_Achievements_OnAchievementsUnlockedCallbackV2Info* Data)
+{
+	const FAccountId LocalAccountId = FindAccountId(Data->UserId);
+	if (LocalAccountId.IsValid())
+	{
+		const FString AchievementId = UTF8_TO_TCHAR(Data->AchievementId);
+		const FDateTime UnlockTime = FDateTime::FromUnixTimestamp(Data->UnlockTime);
+
+		UE_LOG(LogOnlineServices, Verbose, TEXT("%hs User=%s AchievId=%s"), __FUNCTION__, *ToLogString(LocalAccountId), *AchievementId);
+
+		if (FAchievementStateMap* LocalUserAchievementStates = AchievementStates.Find(LocalAccountId))
+		{
+			if (FAchievementState* AchievementState = LocalUserAchievementStates->Find(AchievementId))
+			{
+				AchievementState->Progress = 1.0f;
+				AchievementState->UnlockTime = UnlockTime;
+			}
+		}
+			
+		FAchievementStateUpdated EventParams;
+		EventParams.LocalAccountId = LocalAccountId;
+		EventParams.AchievementIds.Emplace(AchievementId);
+
+		OnAchievementStateUpdatedEvent.Broadcast(EventParams);
+	}
+}
+
 /* UE::Online */ }
diff --git a/Engine/Plugins/Online/OnlineServicesEOSGS/Source/Private/Online/AuthEOSGS.cpp b/Engine/Plugins/Online/OnlineServicesEOSGS/Source/Private/Online/AuthEOSGS.cpp
--- a/Engine/Plugins/Online/OnlineServicesEOSGS/Source/Private/Online/AuthEOSGS.cpp
+++ b/Engine/Plugins/Online/OnlineServicesEOSGS/Source/Private/Online/AuthEOSGS.cpp
@@ -446,10 +446,10 @@ void FAuthEOSGS::Initialize()
 {
 	Super::Initialize();
 
-	AuthHandle = EOS_Platform_GetAuthInterface(*static_cast<FOnlineServicesEOSGS&>(GetServices()).GetEOSPlatformHandle());
+	AuthHandle = EOS_Platform_GetAuthInterface(*GetServices<FOnlineServicesEOSGS>().GetEOSPlatformHandle());
 	check(AuthHandle != nullptr);
 
-	ConnectHandle = EOS_Platform_GetConnectInterface(*static_cast<FOnlineServicesEOSGS&>(GetServices()).GetEOSPlatformHandle());
+	ConnectHandle = EOS_Platform_GetConnectInterface(*GetServices<FOnlineServicesEOSGS>().GetEOSPlatformHandle());
 	check(ConnectHandle != nullptr);
 
 	RegisterHandlers();
diff --git a/Engine/Plugins/Online/OnlineServicesEOSGS/Source/Private/Online/ExternalUIEOSGS.cpp b/Engine/Plugins/Online/OnlineServicesEOSGS/Source/Private/Online/ExternalUIEOSGS.cpp
--- a/Engine/Plugins/Online/OnlineServicesEOSGS/Source/Private/Online/ExternalUIEOSGS.cpp
+++ b/Engine/Plugins/Online/OnlineServicesEOSGS/Source/Private/Online/ExternalUIEOSGS.cpp
@@ -21,7 +21,7 @@ void FExternalUIEOSGS::Initialize()
 	IEOSSDKManager* SDKManager = IEOSSDKManager::Get();
 	check(SDKManager);
 
-	UIInterfaceHandle = EOS_Platform_GetUIInterface(*static_cast<FOnlineServicesEOSGS&>(GetServices()).GetEOSPlatformHandle());
+	UIInterfaceHandle = EOS_Platform_GetUIInterface(*GetServices<FOnlineServicesEOSGS>().GetEOSPlatformHandle());
 	check(UIInterfaceHandle != nullptr);
 
 	RegisterEventHandlers();
diff --git a/Engine/Plugins/Online/OnlineServicesEOSGS/Source/Private/Online/LeaderboardsEOSGS.cpp b/Engine/Plugins/Online/OnlineServicesEOSGS/Source/Private/Online/LeaderboardsEOSGS.cpp
--- a/Engine/Plugins/Online/OnlineServicesEOSGS/Source/Private/Online/LeaderboardsEOSGS.cpp
+++ b/Engine/Plugins/Online/OnlineServicesEOSGS/Source/Private/Online/LeaderboardsEOSGS.cpp
@@ -25,7 +25,7 @@ void FLeaderboardsEOSGS::Initialize()
 {
 	Super::Initialize();
 
-	LeaderboardsHandle = EOS_Platform_GetLeaderboardsInterface(*static_cast<FOnlineServicesEOSGS&>(GetServices()).GetEOSPlatformHandle());
+	LeaderboardsHandle = EOS_Platform_GetLeaderboardsInterface(*GetServices<FOnlineServicesEOSGS>().GetEOSPlatformHandle());
 	check(LeaderboardsHandle);
 }
 
diff --git a/Engine/Plugins/Online/OnlineServicesEOSGS/Source/Private/Online/LobbiesEOSGS.cpp b/Engine/Plugins/Online/OnlineServicesEOSGS/Source/Private/Online/LobbiesEOSGS.cpp
--- a/Engine/Plugins/Online/OnlineServicesEOSGS/Source/Private/Online/LobbiesEOSGS.cpp
+++ b/Engine/Plugins/Online/OnlineServicesEOSGS/Source/Private/Online/LobbiesEOSGS.cpp
@@ -36,7 +36,7 @@ void FLobbiesEOSGS::Initialize()
 	IEOSSDKManager* SDKManager = IEOSSDKManager::Get();
 	check(SDKManager);
 
-	EOS_HLobby LobbyInterfaceHandle = EOS_Platform_GetLobbyInterface(*static_cast<FOnlineServicesEOSGS&>(GetServices()).GetEOSPlatformHandle());
+	EOS_HLobby LobbyInterfaceHandle = EOS_Platform_GetLobbyInterface(*GetServices<FOnlineServicesEOSGS>().GetEOSPlatformHandle());
 	check(LobbyInterfaceHandle != nullptr);
 
 	LobbyPrerequisites = MakeShared<FLobbyPrerequisitesEOS>(FLobbyPrerequisitesEOS{
@@ -52,6 +52,8 @@ void FLobbiesEOSGS::Initialize()
 
 void FLobbiesEOSGS::PreShutdown()
 {
+	Super::PreShutdown();
+
 	UnregisterHandlers();
 
 	LobbyDataRegistry = nullptr;
diff --git a/Engine/Plugins/Online/OnlineServicesEOSGS/Source/Private/Online/SessionsEOSGS.cpp b/Engine/Plugins/Online/OnlineServicesEOSGS/Source/Private/Online/SessionsEOSGS.cpp
--- a/Engine/Plugins/Online/OnlineServicesEOSGS/Source/Private/Online/SessionsEOSGS.cpp
+++ b/Engine/Plugins/Online/OnlineServicesEOSGS/Source/Private/Online/SessionsEOSGS.cpp
@@ -144,7 +144,7 @@ void FSessionsEOSGS::Initialize()
 {
 	Super::Initialize();
 
-	SessionsHandle = EOS_Platform_GetSessionsInterface(*static_cast<FOnlineServicesEOSGS&>(GetServices()).GetEOSPlatformHandle());
+	SessionsHandle = EOS_Platform_GetSessionsInterface(*GetServices<FOnlineServicesEOSGS>().GetEOSPlatformHandle());
 	check(SessionsHandle);
 
 	RegisterEventHandlers();
diff --git a/Engine/Plugins/Online/OnlineServicesEOSGS/Source/Private/Online/StatsEOSGS.cpp b/Engine/Plugins/Online/OnlineServicesEOSGS/Source/Private/Online/StatsEOSGS.cpp
--- a/Engine/Plugins/Online/OnlineServicesEOSGS/Source/Private/Online/StatsEOSGS.cpp
+++ b/Engine/Plugins/Online/OnlineServicesEOSGS/Source/Private/Online/StatsEOSGS.cpp
@@ -30,7 +30,7 @@ void FStatsEOSGS::Initialize()
 {
 	Super::Initialize();
 
-	StatsHandle = EOS_Platform_GetStatsInterface(*static_cast<FOnlineServicesEOSGS&>(GetServices()).GetEOSPlatformHandle());
+	StatsHandle = EOS_Platform_GetStatsInterface(*GetServices<FOnlineServicesEOSGS>().GetEOSPlatformHandle());
 	check(StatsHandle);
 }
 
diff --git a/Engine/Plugins/Online/OnlineServicesEOSGS/Source/Private/Online/TitleFileEOSGS.cpp b/Engine/Plugins/Online/OnlineServicesEOSGS/Source/Private/Online/TitleFileEOSGS.cpp
--- a/Engine/Plugins/Online/OnlineServicesEOSGS/Source/Private/Online/TitleFileEOSGS.cpp
+++ b/Engine/Plugins/Online/OnlineServicesEOSGS/Source/Private/Online/TitleFileEOSGS.cpp
@@ -23,7 +23,7 @@ void FTitleFileEOSGS::Initialize()
 {
 	Super::Initialize();
 
-	TitleStorageHandle = EOS_Platform_GetTitleStorageInterface(*static_cast<FOnlineServicesEOSGS&>(GetServices()).GetEOSPlatformHandle());
+	TitleStorageHandle = EOS_Platform_GetTitleStorageInterface(*GetServices<FOnlineServicesEOSGS>().GetEOSPlatformHandle());
 	check(TitleStorageHandle);
 }
 
diff --git a/Engine/Plugins/Online/OnlineServicesEOSGS/Source/Private/Online/UserFileEOSGS.cpp b/Engine/Plugins/Online/OnlineServicesEOSGS/Source/Private/Online/UserFileEOSGS.cpp
--- a/Engine/Plugins/Online/OnlineServicesEOSGS/Source/Private/Online/UserFileEOSGS.cpp
+++ b/Engine/Plugins/Online/OnlineServicesEOSGS/Source/Private/Online/UserFileEOSGS.cpp
@@ -25,7 +25,7 @@ void FUserFileEOSGS::Initialize()
 {
 	Super::Initialize();
 
-	PlayerDataStorageHandle = EOS_Platform_GetPlayerDataStorageInterface(*static_cast<FOnlineServicesEOSGS&>(GetServices()).GetEOSPlatformHandle());
+	PlayerDataStorageHandle = EOS_Platform_GetPlayerDataStorageInterface(*GetServices<FOnlineServicesEOSGS>().GetEOSPlatformHandle());
 	check(PlayerDataStorageHandle);
 }
 
diff --git a/Engine/Plugins/Online/OnlineServicesEOSGS/Source/Public/Online/AchievementsEOSGS.h b/Engine/Plugins/Online/OnlineServicesEOSGS/Source/Public/Online/AchievementsEOSGS.h
--- a/Engine/Plugins/Online/OnlineServicesEOSGS/Source/Public/Online/AchievementsEOSGS.h
+++ b/Engine/Plugins/Online/OnlineServicesEOSGS/Source/Public/Online/AchievementsEOSGS.h
@@ -4,6 +4,7 @@
 
 #include "CoreMinimal.h"
 #include "Online/AchievementsCommon.h"
+#include "Online/OnlineServicesEOSGSTypes.h"
 
 #if defined(EOS_PLATFORM_BASE_FILE_NAME)
 #include EOS_PLATFORM_BASE_FILE_NAME
@@ -35,13 +36,14 @@ public:
 	virtual TOnlineAsyncOpHandle<FUnlockAchievements> UnlockAchievements(FUnlockAchievements::Params&& Params) override;
 
 protected:
+	void HandleAchievementsUnlocked(const EOS_Achievements_OnAchievementsUnlockedCallbackV2Info* Data);
+
 	EOS_HAchievements AchievementsHandle = nullptr;
 
 	using FAchievementDefinitionMap = TMap<FString, FAchievementDefinition>;
-
 	TOptional<FAchievementDefinitionMap> AchievementDefinitions;
 
-	EOS_NotificationId NotifyAchievementsUnlockedNotificationId = EOS_INVALID_NOTIFICATIONID;
+	FEOSEventRegistrationPtr OnAchievementsUnlocked;
 };
 
 /* UE::Online */ }
diff --git a/Engine/Plugins/Online/OnlineServicesEOSGS/Source/Public/Online/AuthEOSGS.h b/Engine/Plugins/Online/OnlineServicesEOSGS/Source/Public/Online/AuthEOSGS.h
--- a/Engine/Plugins/Online/OnlineServicesEOSGS/Source/Public/Online/AuthEOSGS.h
+++ b/Engine/Plugins/Online/OnlineServicesEOSGS/Source/Public/Online/AuthEOSGS.h
@@ -271,9 +271,9 @@ protected:
 
 	EOS_HAuth AuthHandle = nullptr;
 	EOS_HConnect ConnectHandle = nullptr;
-	EOSEventRegistrationPtr OnConnectLoginStatusChangedEOSEventRegistration;
-	EOSEventRegistrationPtr OnConnectAuthNotifyExpirationEOSEventRegistration;
-	EOSEventRegistrationPtr OnAuthLoginStatusChangedEOSEventRegistration;
+	FEOSEventRegistrationPtr OnConnectLoginStatusChangedEOSEventRegistration;
+	FEOSEventRegistrationPtr OnConnectAuthNotifyExpirationEOSEventRegistration;
+	FEOSEventRegistrationPtr OnAuthLoginStatusChangedEOSEventRegistration;
 	FAccountInfoRegistryEOS AccountInfoRegistryEOS;
 	TLocalUserArray<FUserScopedData> UserScopedData;
 };
diff --git a/Engine/Plugins/Online/OnlineServicesEOSGS/Source/Public/Online/ExternalUIEOSGS.h b/Engine/Plugins/Online/OnlineServicesEOSGS/Source/Public/Online/ExternalUIEOSGS.h
--- a/Engine/Plugins/Online/OnlineServicesEOSGS/Source/Public/Online/ExternalUIEOSGS.h
+++ b/Engine/Plugins/Online/OnlineServicesEOSGS/Source/Public/Online/ExternalUIEOSGS.h
@@ -51,7 +51,7 @@ protected:
 	TOnlineAsyncOpHandle<FExternalUIProcessDisplaySettingsUpdatedImp> ProcessDisplaySettingsUpdatedImplOp(FExternalUIProcessDisplaySettingsUpdatedImp::Params&& Params);
 
 	EOS_HUI UIInterfaceHandle = nullptr;
-	EOSEventRegistrationPtr OnDisplaySettingsUpdated;
+	FEOSEventRegistrationPtr OnDisplaySettingsUpdated;
 };
 
 namespace Meta {
diff --git a/Engine/Plugins/Online/OnlineServicesEOSGS/Source/Public/Online/LobbiesEOSGS.h b/Engine/Plugins/Online/OnlineServicesEOSGS/Source/Public/Online/LobbiesEOSGS.h
--- a/Engine/Plugins/Online/OnlineServicesEOSGS/Source/Public/Online/LobbiesEOSGS.h
+++ b/Engine/Plugins/Online/OnlineServicesEOSGS/Source/Public/Online/LobbiesEOSGS.h
@@ -187,6 +187,8 @@ struct FLobbiesProcessLobbyNotificationImpl
 class FLobbiesEOSGS : public FLobbiesCommon
 {
 public:
+	using Super = FLobbiesCommon;
+
 	FLobbiesEOSGS(FOnlineServicesEOSGS& InServices);
 
 	virtual void Initialize() override;
@@ -239,12 +241,12 @@ protected:
 	TFuture<TDefaultErrorResult<FLobbiesModifyLobbyDataImpl>> ModifyLobbyDataImpl(FLobbiesModifyLobbyDataImpl::Params&& Params);
 	TOnlineAsyncOpHandle<FLobbiesProcessLobbyNotificationImpl> ProcessLobbyNotificationImplOp(FLobbiesProcessLobbyNotificationImpl::Params&& Params);
 
-	EOSEventRegistrationPtr OnLobbyUpdatedEOSEventRegistration;
-	EOSEventRegistrationPtr OnLobbyMemberUpdatedEOSEventRegistration;
-	EOSEventRegistrationPtr OnLobbyMemberStatusReceivedEOSEventRegistration;
-	EOSEventRegistrationPtr OnLobbyInviteReceivedEOSEventRegistration;
-	EOSEventRegistrationPtr OnLobbyInviteAcceptedEOSEventRegistration;
-	EOSEventRegistrationPtr OnJoinLobbyAcceptedEOSEventRegistration;
+	FEOSEventRegistrationPtr OnLobbyUpdatedEOSEventRegistration;
+	FEOSEventRegistrationPtr OnLobbyMemberUpdatedEOSEventRegistration;
+	FEOSEventRegistrationPtr OnLobbyMemberStatusReceivedEOSEventRegistration;
+	FEOSEventRegistrationPtr OnLobbyInviteReceivedEOSEventRegistration;
+	FEOSEventRegistrationPtr OnLobbyInviteAcceptedEOSEventRegistration;
+	FEOSEventRegistrationPtr OnJoinLobbyAcceptedEOSEventRegistration;
 
 	TSharedPtr<FLobbyPrerequisitesEOS> LobbyPrerequisites;
 	TSharedPtr<FLobbyDataRegistryEOS> LobbyDataRegistry;
diff --git a/Engine/Plugins/Online/OnlineServicesEOSGS/Source/Public/Online/OnlineServicesEOSGSTypes.h b/Engine/Plugins/Online/OnlineServicesEOSGS/Source/Public/Online/OnlineServicesEOSGSTypes.h
--- a/Engine/Plugins/Online/OnlineServicesEOSGS/Source/Public/Online/OnlineServicesEOSGSTypes.h
+++ b/Engine/Plugins/Online/OnlineServicesEOSGS/Source/Public/Online/OnlineServicesEOSGSTypes.h
@@ -117,12 +117,15 @@ decltype(auto) EOS_Async(TEOSFn EOSFn, TEOSHandle EOSHandle, const TEOSParameter
 	}));
 }
 
-class EOSEventRegistration
+class FEOSEventRegistration
 {
 public:
-	virtual ~EOSEventRegistration() = default;
+	virtual ~FEOSEventRegistration() = default;
 };
-typedef TUniquePtr<EOSEventRegistration> EOSEventRegistrationPtr;
+typedef TUniquePtr<FEOSEventRegistration> FEOSEventRegistrationPtr;
+
+using EOSEventRegistration UE_DEPRECATED(5.6, "EOSEventRegistration renamed to FEOSEventRegistration") = FEOSEventRegistration;
+using EOSEventRegistrationPtr UE_DEPRECATED(5.6, "EOSEventRegistrationPtr renamed to FEOSEventRegistrationPtr") = FEOSEventRegistrationPtr;
 
 /** 
 * EOS event registration utility for binding an EOS notifier registration to a RAII object which handles
@@ -143,7 +146,7 @@ template <
 	typename EOSNotfyRegisterFunction,
 	typename EOSNotfyUnregisterFunction,
 	typename ComponentHandlerFunction>
-EOSEventRegistrationPtr EOS_RegisterComponentEventHandler(
+FEOSEventRegistrationPtr EOS_RegisterComponentEventHandler(
 	ComponentHandlerClass* HandlerClass,
 	EOSHandle ClientHandle,
 	int32_t ApiVersion,
@@ -182,17 +185,17 @@ template <
 	typename EOSNotfyRegisterFunction,
 	typename EOSNotfyUnregisterFunction,
 	typename ComponentHandlerFunction>
-class EOSEventRegistrationImpl : public EOSEventRegistration
+class FEOSEventRegistrationImpl : public FEOSEventRegistration
 {
 public:
-	typedef EOSEventRegistrationImpl<
+	typedef FEOSEventRegistrationImpl<
 		ComponentHandlerClass,
 		EOSHandle,
 		EOSNotfyRegisterFunction,
 		EOSNotfyUnregisterFunction,
 		ComponentHandlerFunction> ThisClass;
 
-	EOSEventRegistrationImpl(
+	FEOSEventRegistrationImpl(
 		ComponentHandlerClass* HandlerClass,
 		EOSHandle ClientHandle,
 		int32_t ApiVersion,
@@ -215,15 +218,15 @@ public:
 		}));
 	}
 
-	virtual ~EOSEventRegistrationImpl()
+	virtual ~FEOSEventRegistrationImpl()
 	{
 		NotfyUnregisterFunction(ClientHandle, NotificationId);
 	}
 
 private:
-	EOSEventRegistrationImpl() = delete;
-	EOSEventRegistrationImpl(const EOSEventRegistrationImpl&) = delete;
-	EOSEventRegistrationImpl& operator=(const EOSEventRegistrationImpl&) = delete;
+	FEOSEventRegistrationImpl() = delete;
+	FEOSEventRegistrationImpl(const FEOSEventRegistrationImpl&) = delete;
+	FEOSEventRegistrationImpl& operator=(const FEOSEventRegistrationImpl&) = delete;
 
 	typedef typename TEOSNotifyRegisterTraits<EOSNotfyRegisterFunction>::NotificationIdType NotificationIdType;
 
@@ -242,7 +245,7 @@ template <
 	typename EOSNotfyRegisterFunction,
 	typename EOSNotfyUnregisterFunction,
 	typename ComponentHandlerFunction>
-EOSEventRegistrationPtr EOS_RegisterComponentEventHandler(
+FEOSEventRegistrationPtr EOS_RegisterComponentEventHandler(
 	ComponentHandlerClass* HandlerClass,
 	EOSHandle ClientHandle,
 	int32_t ApiVersion,
@@ -251,7 +254,7 @@ EOSEventRegistrationPtr EOS_RegisterComponentEventHandler(
 	ComponentHandlerFunction HandlerFunction)
 {
 	return MakeUnique<
-		Private::EOSEventRegistrationImpl<
+		Private::FEOSEventRegistrationImpl<
 			ComponentHandlerClass,
 			EOSHandle,
 			EOSNotfyRegisterFunction,
diff --git a/Engine/Plugins/Online/OnlineServicesEOSGS/Source/Public/Online/SessionsEOSGS.h b/Engine/Plugins/Online/OnlineServicesEOSGS/Source/Public/Online/SessionsEOSGS.h
--- a/Engine/Plugins/Online/OnlineServicesEOSGS/Source/Public/Online/SessionsEOSGS.h
+++ b/Engine/Plugins/Online/OnlineServicesEOSGS/Source/Public/Online/SessionsEOSGS.h
@@ -282,9 +282,9 @@ private:
 protected:
 	EOS_HSessions SessionsHandle = nullptr;
 
-	EOSEventRegistrationPtr OnSessionInviteReceivedEventRegistration;
-	EOSEventRegistrationPtr OnSessionInviteAcceptedEventRegistration;
-	EOSEventRegistrationPtr OnJoinSessionAcceptedEventRegistration;
+	FEOSEventRegistrationPtr OnSessionInviteReceivedEventRegistration;
+	FEOSEventRegistrationPtr OnSessionInviteAcceptedEventRegistration;
+	FEOSEventRegistrationPtr OnJoinSessionAcceptedEventRegistration;
 
 	TMap<FAccountId, TSharedRef<FSessionSearchHandleEOSGS>> CurrentSessionSearchHandleEOSGSUserMap;
 };
