Added validity checks and handling instead of assuming parameters are valid in FPresenceEOS::UpdateUserPresence

diff --git a/Engine/Plugins/Online/OnlineServicesEOS/Source/Private/Online/PresenceEOS.cpp b/Engine/Plugins/Online/OnlineServicesEOS/Source/Private/Online/PresenceEOS.cpp
--- a/Engine/Plugins/Online/OnlineServicesEOS/Source/Private/Online/PresenceEOS.cpp
+++ b/Engine/Plugins/Online/OnlineServicesEOS/Source/Private/Online/PresenceEOS.cpp
@@ -609,6 +609,20 @@ TSharedRef<FUserPresence> FPresenceEOS::FindOrCreatePresence(FAccountId LocalAcc
 
 void FPresenceEOS::UpdateUserPresence(FAccountId LocalAccountId, FAccountId PresenceAccountId)
 {
+	const EOS_EpicAccountId LocalUserId = GetEpicAccountId(LocalAccountId);
+	if (LocalUserId == nullptr)
+	{
+		UE_LOG(LogOnlineServices, Warning, TEXT("[%hs] No EpicAccountId found for FAccountId [%s]"), __FUNCTION__, *ToLogString(LocalAccountId));
+		return;
+	}
+
+	const EOS_EpicAccountId TargetUserId = GetEpicAccountId(PresenceAccountId);
+	if (TargetUserId == nullptr)
+	{
+		UE_LOG(LogOnlineServices, Warning, TEXT("[%hs] No EpicAccountId found for FAccountId [%s]"), __FUNCTION__, *ToLogString(PresenceAccountId));
+		return;
+	}
+
 	bool bPresenceHasChanged = false;
 	TSharedRef<FUserPresence> UserPresence = FindOrCreatePresence(LocalAccountId, PresenceAccountId);
 	// TODO:  Handle updates for local users.  Don't want to conflict with UpdatePresence calls
@@ -618,8 +632,8 @@ void FPresenceEOS::UpdateUserPresence(FAccountId LocalAccountId, FAccountId Pres
 	EOS_Presence_CopyPresenceOptions Options = { };
 	Options.ApiVersion = 3;
 	UE_EOS_CHECK_API_MISMATCH(EOS_PRESENCE_COPYPRESENCE_API_LATEST, 3);
-	Options.LocalUserId = GetEpicAccountIdChecked(LocalAccountId);
-	Options.TargetUserId = GetEpicAccountIdChecked(PresenceAccountId);
+	Options.LocalUserId = LocalUserId;
+	Options.TargetUserId = TargetUserId;
 	EOS_EResult CopyPresenceResult = EOS_Presence_CopyPresence(PresenceHandle, &Options, &PresenceInfo);
 	if (CopyPresenceResult == EOS_EResult::EOS_Success)
 	{
