[AutoRTFM] Use a transactionally-safe RWLock in UE::Online::FOnlineAccountIdRegistryEOS.

This is necessary because this code is reachable from a Verse transaction, such as when logging AccountID information.

diff --git a/Engine/Plugins/Online/OnlineServicesEOS/Source/Private/Online/OnlineIdEOS.cpp b/Engine/Plugins/Online/OnlineServicesEOS/Source/Private/Online/OnlineIdEOS.cpp
--- a/Engine/Plugins/Online/OnlineServicesEOS/Source/Private/Online/OnlineIdEOS.cpp
+++ b/Engine/Plugins/Online/OnlineServicesEOS/Source/Private/Online/OnlineIdEOS.cpp
@@ -142,7 +142,7 @@ FAccountId FOnlineAccountIdRegistryEOS::FindAccountId(const EOS_ProductUserId Pr
 	FAccountId AccountId;
 	if (EOS_ProductUserId_IsValid(ProductUserId) == EOS_TRUE)
 	{
-		const FReadScopeLock ReadLock(Lock);
+		const UE::TReadScopeLock ReadLock(Lock);
 		if (const FAccountId* Found = ProductUserIdToAccountId.Find(ProductUserId))
 		{
 			AccountId = *Found;
@@ -161,7 +161,7 @@ FAccountId FOnlineAccountIdRegistryEOS::FindAccountId(const EOS_EpicAccountId Ep
 	FAccountId AccountId;
 	if (EOS_EpicAccountId_IsValid(EpicAccountId) == EOS_TRUE)
 	{
-		const FReadScopeLock ReadLock(Lock);
+		const UE::TReadScopeLock ReadLock(Lock);
 		if (const FAccountId* Found = EpicAccountIdToAccountId.Find(EpicAccountId))
 		{
 			AccountId = *Found;
@@ -174,7 +174,7 @@ FOnlineAccountIdDataEOS FOnlineAccountIdRegistryEOS::GetAccountIdData(const FAcc
 {
 	if (ValidateOnlineId(AccountId))
 	{
-		FReadScopeLock ScopeLock(Lock);
+		UE::TReadScopeLock ScopeLock(Lock);
 		return AccountIdData[AccountId.GetHandle() - 1];
 	}
 	return FOnlineAccountIdDataEOS();
@@ -226,14 +226,14 @@ FAccountId FOnlineAccountIdRegistryEOS::FindOrAddAccountId(const EOS_EpicAccount
 
 		{
 			// First take read lock and look for existing elements
-			const FReadScopeLock ReadLock(Lock);
+			const UE::TReadScopeLock ReadLock(Lock);
 			Result = FindExisting();
 		}
 
 		if(!Result.IsValid())
 		{
 			// Double-checked locking. If we didn't find an element, we take the write lock, and look again, in case another thread raced with us and added one.
-			const FWriteScopeLock WriteLock(Lock);
+			const UE::TWriteScopeLock WriteLock(Lock);
 			Result = FindExisting();
 
 			if (!Result.IsValid())
@@ -260,7 +260,7 @@ FAccountId FOnlineAccountIdRegistryEOS::FindOrAddAccountId(const EOS_EpicAccount
 		if (bUpdateEpicAccountId || bUpdateProductUserId)
 		{
 			// Finally, update any previously unset fields for which we now have a valid value.
-			const FWriteScopeLock WriteLock(Lock);
+			const UE::TWriteScopeLock WriteLock(Lock);
 			FOnlineAccountIdDataEOS& AccountIdDataToUpdate = AccountIdData[Result.GetHandle() -1];
 			if (bUpdateEpicAccountId)
 			{
diff --git a/Engine/Plugins/Online/OnlineServicesEOS/Source/Private/Online/OnlineIdEOS.h b/Engine/Plugins/Online/OnlineServicesEOS/Source/Private/Online/OnlineIdEOS.h
--- a/Engine/Plugins/Online/OnlineServicesEOS/Source/Private/Online/OnlineIdEOS.h
+++ b/Engine/Plugins/Online/OnlineServicesEOS/Source/Private/Online/OnlineIdEOS.h
@@ -4,6 +4,7 @@
 
 #include "Online/OnlineIdEOSGS.h"
 #include "Online/AccountIdEOS.h"
+#include "Misc/TransactionallySafeRWLock.h"
 
 struct EOS_EpicAccountIdDetails;
 struct EOS_ProductUserIdDetails;
@@ -59,7 +60,7 @@ private:
 	friend class FOnlineServicesEOSModule;
 	friend FAccountIdReplicationTest;
 
-	mutable FRWLock Lock;
+	mutable FTransactionallySafeRWLock Lock;
 
 	TArray<FOnlineAccountIdDataEOS> AccountIdData; // Actual container for the info, indexed by the handle
 
