Addition of support for EOS_Ecom_QueryOwnershipBySandboxIds, called after EOS_Ecom_QueryEntitlements to return full list of entitlements (consumable and durable)

diff --git a/Engine/Plugins/Online/OnlineServicesEOS/Source/Private/Online/CommerceEOS.cpp b/Engine/Plugins/Online/OnlineServicesEOS/Source/Private/Online/CommerceEOS.cpp
--- a/Engine/Plugins/Online/OnlineServicesEOS/Source/Private/Online/CommerceEOS.cpp
+++ b/Engine/Plugins/Online/OnlineServicesEOS/Source/Private/Online/CommerceEOS.cpp
@@ -396,6 +396,8 @@ TOnlineAsyncOpHandle<FCommerceQueryEntitlements> FCommerceEOS::QueryEntitlements
 			TArray<FEntitlement>& Entitlements = CachedEntitlements.FindOrAdd(Op.GetParams().LocalAccountId, TArray<FEntitlement>());
 			Entitlements.Reset(Count);
 
+			UE_LOG(LogOnlineServices, VeryVerbose, TEXT("Entitlements Count: %d"), Count);
+
 			EOS_Ecom_CopyEntitlementByIndexOptions CopyOptions = { };
 			CopyOptions.ApiVersion = 1;
 			UE_EOS_CHECK_API_MISMATCH(EOS_ECOM_COPYENTITLEMENTBYINDEX_API_LATEST, 1);
@@ -405,6 +407,8 @@ TOnlineAsyncOpHandle<FCommerceQueryEntitlements> FCommerceEOS::QueryEntitlements
 			{
 				CopyOptions.EntitlementIndex = Index;
 
+				UE_LOG(LogOnlineServices, VeryVerbose, TEXT("	Entitlement %d:"), Index);
+
 				EOS_Ecom_Entitlement* EosEntitlement = nullptr;
 				EOS_EResult CopyResult = EOS_Ecom_CopyEntitlementByIndex(EcomHandle, &CopyOptions, &EosEntitlement);
 				if (CopyResult != EOS_EResult::EOS_Success && CopyResult != EOS_EResult::EOS_Ecom_EntitlementStale)
@@ -417,8 +421,66 @@ TOnlineAsyncOpHandle<FCommerceQueryEntitlements> FCommerceEOS::QueryEntitlements
 				FEntitlement& Entitlement = Entitlements.Emplace_GetRef();
 				EOSEntitlementToOssEntitlement(Entitlement, EosEntitlement);
 
+				UE_LOG(LogOnlineServices, VeryVerbose, TEXT("		Entitlement Id: %s"), *Entitlement.EntitlementId);
+				UE_LOG(LogOnlineServices, VeryVerbose, TEXT("		Product Id: %s"), *Entitlement.ProductId);
+
 				EOS_Ecom_Entitlement_Release(EosEntitlement);
 			}
+		})
+		.Then([this](TOnlineAsyncOp<FCommerceQueryEntitlements>& Op, TPromise<const EOS_Ecom_QueryOwnershipBySandboxIdsCallbackInfo*>&& Promise)
+		{
+			const FCommerceQueryEntitlements::Params& Params = Op.GetParams();
+
+			EOS_Ecom_QueryOwnershipBySandboxIdsOptions Options = {};
+			Options.ApiVersion = 1;
+			UE_EOS_CHECK_API_MISMATCH(EOS_ECOM_QUERYOWNERSHIPBYSANDBOXIDSOPTIONS_API_LATEST, 1);
+			Options.LocalUserId = GetEpicAccountIdChecked(Params.LocalAccountId);
+			Options.SandboxIdsCount = 1;
+
+			IEOSSDKManager* Manager = IEOSSDKManager::Get();
+			const FString& PlatformConfigName = StaticCast<FOnlineServicesEpicCommon*>(&Services)->GetEOSPlatformHandle()->GetConfigName();
+			const FEOSSDKPlatformConfig* Config = Manager ? Manager->GetPlatformConfig(PlatformConfigName) : nullptr;
+
+			const FTCHARToUTF8& Converter = FTCHARToUTF8(*Config->SandboxId);
+			const char* SandboxId = Converter.Get();
+			Options.SandboxIds = &SandboxId;
+			
+			EOS_Async(EOS_Ecom_QueryOwnershipBySandboxIds, EcomHandle, Options, MoveTemp(Promise));
+		})
+		.Then([this](TOnlineAsyncOp<FCommerceQueryEntitlements>& Op, const EOS_Ecom_QueryOwnershipBySandboxIdsCallbackInfo* Data)
+		{
+			const EOS_EResult Result = Data->ResultCode;
+			if (Result != EOS_EResult::EOS_Success)
+			{
+				UE_LOG(LogOnlineServices, Error, TEXT("EOS_Ecom_QueryOwnership: failed with error (%s)"), *LexToString(Data->ResultCode));
+				Op.SetError(Errors::FromEOSResult(Result));
+				return;
+			}
+
+			const FCommerceQueryEntitlements::Params& Params = Op.GetParams();
+			
+			TArray<FEntitlement>& Entitlements = CachedEntitlements.FindOrAdd(Params.LocalAccountId, TArray<FEntitlement>());
+			
+			UE_LOG(LogOnlineServices, VeryVerbose, TEXT("Sandbox Id Ownership Count: %d"), Data->SandboxIdItemOwnershipsCount);
+
+			for (const EOS_Ecom_SandboxIdItemOwnership& SandboxOwnershipEntry : TArrayView<const EOS_Ecom_SandboxIdItemOwnership>(Data->SandboxIdItemOwnerships, Data->SandboxIdItemOwnershipsCount))
+			{
+				UE_LOG(LogOnlineServices, VeryVerbose, TEXT("	Sandbox Id: %hs"), SandboxOwnershipEntry.SandboxId);
+				UE_LOG(LogOnlineServices, VeryVerbose, TEXT("	Ownership Count: %d"), SandboxOwnershipEntry.OwnedCatalogItemIdsCount);
+
+				for (uint32_t CatalogItemIndex = 0; CatalogItemIndex < SandboxOwnershipEntry.OwnedCatalogItemIdsCount; CatalogItemIndex++)
+				{
+					const EOS_Ecom_CatalogItemId& CatalogItemId = SandboxOwnershipEntry.OwnedCatalogItemIds[CatalogItemIndex];
+
+					FEntitlement& Entitlement = Entitlements.Emplace_GetRef();
+					Entitlement.bRedeemed = true;
+					Entitlement.ProductId = FString(CatalogItemId);
+					Entitlement.Quantity = 1;
+
+					UE_LOG(LogOnlineServices, VeryVerbose, TEXT("		Catalog Item %d"), CatalogItemIndex);
+					UE_LOG(LogOnlineServices, VeryVerbose, TEXT("			Catalog Item Id: %hs"), CatalogItemId);
+				}
+			}
 
 			Op.SetResult({});
 		})
diff --git a/Engine/Plugins/Online/OnlineSubsystemEOS/Source/OnlineSubsystemEOS/Private/OnlineStoreEOS.cpp b/Engine/Plugins/Online/OnlineSubsystemEOS/Source/OnlineSubsystemEOS/Private/OnlineStoreEOS.cpp
--- a/Engine/Plugins/Online/OnlineSubsystemEOS/Source/OnlineSubsystemEOS/Private/OnlineStoreEOS.cpp
+++ b/Engine/Plugins/Online/OnlineSubsystemEOS/Source/OnlineSubsystemEOS/Private/OnlineStoreEOS.cpp
@@ -4,6 +4,7 @@
 
 #if WITH_EOS_SDK
 
+#include "IEOSSDKManager.h"
 #include "OnlineSubsystemEOS.h"
 #include "OnlineSubsystemEOSPrivate.h"
 #include "UserManagerEOS.h"
@@ -298,6 +299,7 @@ void FOnlineStoreEOS::RedeemCode(const FUniqueNetId& UserId, const FRedeemCodeRe
 }
 
 typedef TEOSCallback<EOS_Ecom_OnQueryEntitlementsCallback, EOS_Ecom_QueryEntitlementsCallbackInfo, FOnlineStoreEOS> FQueryReceiptsCallback;
+typedef TEOSCallback<EOS_Ecom_OnQueryOwnershipBySandboxIdsCallback, EOS_Ecom_QueryOwnershipBySandboxIdsCallbackInfo, FOnlineStoreEOS> FQueryOwnershipBySandboxIdsCallback;
 
 void FOnlineStoreEOS::QueryReceipts(const FUniqueNetId& UserId, bool bRestoreReceipts, const FOnQueryReceiptsComplete& Delegate)
 {
@@ -312,6 +314,8 @@ void FOnlineStoreEOS::QueryReceipts(const FUniqueNetId& UserId, bool bRestoreRec
 
 	CachedReceipts.Reset();
 
+	// First we'll call QueryEntitlements to query ownership for individual products
+
 	EOS_Ecom_QueryEntitlementsOptions Options = { };
 	Options.ApiVersion = 3;
 	UE_EOS_CHECK_API_MISMATCH(EOS_ECOM_QUERYENTITLEMENTS_API_LATEST, 3);
@@ -367,7 +371,57 @@ void FOnlineStoreEOS::QueryReceipts(const FUniqueNetId& UserId, bool bRestoreRec
 			EOS_Ecom_Entitlement_Release(Receipt);
 		}
 
-		OnComplete.ExecuteIfBound(ONLINE_ERROR(EOnlineErrorResult::Success));
+		// After QueryEntitlements is done, we call QueryOwnership to query ownership granted by other services like season passes, any benefits of which would not have been returned by QueryEntitlements
+
+		EOS_Ecom_QueryOwnershipBySandboxIdsOptions Options = {};
+		Options.ApiVersion = 1;
+		UE_EOS_CHECK_API_MISMATCH(EOS_ECOM_QUERYOWNERSHIPBYSANDBOXIDSOPTIONS_API_LATEST, 1);
+		Options.LocalUserId = Data->LocalUserId;
+		Options.SandboxIdsCount = 1;
+
+		IEOSSDKManager* Manager = IEOSSDKManager::Get();
+		const FString& PlatformConfigName = EOSSubsystem->GetEOSPlatformHandle()->GetConfigName();
+		const FEOSSDKPlatformConfig* Config = Manager ? Manager->GetPlatformConfig(PlatformConfigName) : nullptr;
+
+		const FTCHARToUTF8& Converter = FTCHARToUTF8(*Config->SandboxId);
+		const char* SandboxId = Converter.Get();
+		Options.SandboxIds = &SandboxId;
+
+		FQueryOwnershipBySandboxIdsCallback* CallbackObj = new FQueryOwnershipBySandboxIdsCallback(FOnlineStoreEOSWeakPtr(AsShared()));
+		CallbackObj->CallbackLambda = [this, OnComplete = FOnQueryReceiptsComplete(OnComplete)](const EOS_Ecom_QueryOwnershipBySandboxIdsCallbackInfo* Data)
+		{
+			const EOS_EResult Result = Data->ResultCode;
+			if (Result != EOS_EResult::EOS_Success)
+			{
+				UE_LOG_ONLINE(Error, TEXT("EOS_Ecom_QueryOwnership: failed with error (%s)"), *LexToString(Data->ResultCode));
+				OnComplete.ExecuteIfBound(ONLINE_ERROR(EOnlineErrorResult::Unknown));
+				return;
+			}
+			
+			UE_LOG_ONLINE(VeryVerbose, TEXT("Sandbox Id Ownership Count: %d"), Data->SandboxIdItemOwnershipsCount);
+
+			for (const EOS_Ecom_SandboxIdItemOwnership& SandboxOwnershipEntry : TArrayView<const EOS_Ecom_SandboxIdItemOwnership>(Data->SandboxIdItemOwnerships, Data->SandboxIdItemOwnershipsCount))
+			{
+				UE_LOG_ONLINE(VeryVerbose, TEXT("	Sandbox Id: %hs"), SandboxOwnershipEntry.SandboxId);
+				UE_LOG_ONLINE(VeryVerbose, TEXT("	Ownership Count: %d"), SandboxOwnershipEntry.OwnedCatalogItemIdsCount);
+
+				for (uint32_t CatalogItemIndex = 0; CatalogItemIndex < SandboxOwnershipEntry.OwnedCatalogItemIdsCount; CatalogItemIndex++)
+				{
+					const EOS_Ecom_CatalogItemId& CatalogItemId = SandboxOwnershipEntry.OwnedCatalogItemIds[CatalogItemIndex];
+
+					FPurchaseReceipt& PurchaseReceipt = CachedReceipts.Emplace_GetRef();
+					PurchaseReceipt.TransactionState = EPurchaseTransactionState::Purchased;
+					PurchaseReceipt.TransactionId = FString(CatalogItemId);
+
+					UE_LOG_ONLINE(VeryVerbose, TEXT("		Catalog Item %d"), CatalogItemIndex);
+					UE_LOG_ONLINE(VeryVerbose, TEXT("			Catalog Item Id: %hs"), CatalogItemId);
+				}
+			}
+			
+			OnComplete.ExecuteIfBound(ONLINE_ERROR(EOnlineErrorResult::Success));
+		};
+
+		EOS_Ecom_QueryOwnershipBySandboxIds(EOSSubsystem->EcomHandle, &Options, CallbackObj, CallbackObj->GetCallbackPtr());
 	};
 	EOS_Ecom_QueryEntitlements(EOSSubsystem->EcomHandle, &Options, CallbackObj, CallbackObj->GetCallbackPtr());
 }
