- Factor resolve APIs out of FAuthEOS[GS] into IEpicAccountIdResolver/IEpicProductUserIdResolver APIs, so that components that are instantiated in both OnlineServicesEOS[GS] and OnlineServicesEpicAccount/Game can call them and "do the right thing".

diff --git a/Engine/Plugins/Online/OnlineServicesEOS/Source/Private/Online/AuthEOS.cpp b/Engine/Plugins/Online/OnlineServicesEOS/Source/Private/Online/AuthEOS.cpp
--- a/Engine/Plugins/Online/OnlineServicesEOS/Source/Private/Online/AuthEOS.cpp
+++ b/Engine/Plugins/Online/OnlineServicesEOS/Source/Private/Online/AuthEOS.cpp
@@ -9,6 +9,8 @@
 #include "EOSShared.h"
 #include "IEOSSDKManager.h"
 #include "Online/AuthErrors.h"
+#include "Online/EpicAccountIdResolver.h"
+#include "Online/EpicProductUserIdResolver.h"
 #include "Online/OnlineErrorEpicCommon.h"
 #include "Online/OnlineIdEOS.h"
 #include "Online/OnlineServicesEOS.h"
@@ -568,25 +570,6 @@ TOnlineAsyncOpHandle<FAuthLinkAccount> FAuthEOS::LinkAccount(FAuthLinkAccount::P
 	return Op->GetHandle();
 }
 
-template <typename IdType>
-TFuture<FAccountId> ResolveAccountIdImpl(FAuthEOS& AuthEOS, const FAccountId& LocalAccountId, const IdType InId)
-{
-	TPromise<FAccountId> Promise;
-	TFuture<FAccountId> Future = Promise.GetFuture();
-
-	AuthEOS.ResolveAccountIds(LocalAccountId, { InId }).Next([Promise = MoveTemp(Promise)](const TArray<FAccountId>& AccountIds) mutable
-	{
-		FAccountId Result;
-		if (AccountIds.Num() == 1)
-		{
-			Result = AccountIds[0];
-		}
-		Promise.SetValue(Result);
-	});
-
-	return Future;
-}
-
 TOnlineAsyncOpHandle<FAuthQueryExternalServerAuthTicket> FAuthEOS::QueryExternalServerAuthTicket(FAuthQueryExternalServerAuthTicket::Params&& Params)
 {
 	TOnlineAsyncOpRef<FAuthQueryExternalServerAuthTicket> Op = GetJoinableOp<FAuthQueryExternalServerAuthTicket>(MoveTemp(Params));
@@ -721,16 +704,6 @@ TOnlineResult<FAuthGetLinkAccountContinuationId> FAuthEOS::GetLinkAccountContinu
 	return TOnlineResult<FAuthGetLinkAccountContinuationId>(FAuthGetLinkAccountContinuationId::Result{UserData->LastLoginContinuationId});
 }
 
-TFuture<FAccountId> FAuthEOS::ResolveAccountId(const FAccountId& LocalAccountId, const EOS_EpicAccountId EpicAccountId)
-{
-	return ResolveAccountIdImpl(*this, LocalAccountId, EpicAccountId);
-}
-
-TFuture<FAccountId> FAuthEOS::ResolveAccountId(const FAccountId& LocalAccountId, const EOS_ProductUserId ProductUserId)
-{
-	return ResolveAccountIdImpl(*this, LocalAccountId, ProductUserId);
-}
-
 using FEpicAccountIdStrBuffer = TStaticArray<char, EOS_EPICACCOUNTID_MAX_LENGTH + 1>;
 
 TFuture<TArray<FAccountId>> FAuthEOS::ResolveAccountIds(const FAccountId& LocalAccountId, const TArray<EOS_EpicAccountId>& InEpicAccountIds)
@@ -926,50 +899,6 @@ TFuture<TArray<FAccountId>> FAuthEOS::ResolveAccountIds(const FAccountId& LocalA
 	return Future;
 }
 
-template<typename ParamType>
-TFunction<TFuture<FAccountId>(FOnlineAsyncOp& InAsyncOp, const ParamType& Param)> ResolveIdFnImpl(FAuthEOS* AuthEOS)
-{
-	return [AuthEOS](FOnlineAsyncOp& InAsyncOp, const ParamType& Param)
-	{
-		const FAccountId* LocalAccountIdPtr = InAsyncOp.Data.Get<FAccountId>(TEXT("LocalAccountId"));
-		if (!ensure(LocalAccountIdPtr))
-		{
-			return MakeFulfilledPromise<FAccountId>().GetFuture();
-		}
-		return AuthEOS->ResolveAccountId(*LocalAccountIdPtr, Param);
-	};
-}
-TFunction<TFuture<FAccountId>(FOnlineAsyncOp& InAsyncOp, const EOS_EpicAccountId& ProductUserId)> FAuthEOS::ResolveEpicIdFn()
-{
-	return ResolveIdFnImpl<EOS_EpicAccountId>(this);
-}
-TFunction<TFuture<FAccountId>(FOnlineAsyncOp& InAsyncOp, const EOS_ProductUserId& ProductUserId)> FAuthEOS::ResolveProductIdFn()
-{
-	return ResolveIdFnImpl<EOS_ProductUserId>(this);
-}
-
-template<typename ParamType>
-TFunction<TFuture<TArray<FAccountId>>(FOnlineAsyncOp& InAsyncOp, const TArray<ParamType>& Param)> ResolveIdsFnImpl(FAuthEOS* AuthEOS)
-{
-	return [AuthEOS](FOnlineAsyncOp& InAsyncOp, const TArray<ParamType>& Param)
-	{
-		const FAccountId* LocalAccountIdPtr = InAsyncOp.Data.Get<FAccountId>(TEXT("LocalAccountId"));
-		if (!ensure(LocalAccountIdPtr))
-		{
-			return MakeFulfilledPromise<TArray<FAccountId>>().GetFuture();
-		}
-		return AuthEOS->ResolveAccountIds(*LocalAccountIdPtr, Param);
-	};
-}
-TFunction<TFuture<TArray<FAccountId>>(FOnlineAsyncOp& InAsyncOp, const TArray<EOS_EpicAccountId>& EpicAccountIds)> FAuthEOS::ResolveEpicIdsFn()
-{
-	return ResolveIdsFnImpl<EOS_EpicAccountId>(this);
-}
-TFunction<TFuture<TArray<FAccountId>>(FOnlineAsyncOp& InAsyncOp, const TArray<EOS_ProductUserId>& ProductUserIds)> FAuthEOS::ResolveProductIdsFn()
-{
-	return ResolveIdsFnImpl<EOS_ProductUserId>(this);
-}
-
 FAccountId FAuthEOS::CreateAccountId(const EOS_EpicAccountId EpicAccountId, const EOS_ProductUserId ProductUserId)
 {
 	return FOnlineAccountIdRegistryEOS::Get().FindOrAddAccountId(EpicAccountId, ProductUserId);
diff --git a/Engine/Plugins/Online/OnlineServicesEOS/Source/Private/Online/CommerceEOS.cpp b/Engine/Plugins/Online/OnlineServicesEOS/Source/Private/Online/CommerceEOS.cpp
--- a/Engine/Plugins/Online/OnlineServicesEOS/Source/Private/Online/CommerceEOS.cpp
+++ b/Engine/Plugins/Online/OnlineServicesEOS/Source/Private/Online/CommerceEOS.cpp
@@ -60,7 +60,7 @@ TOnlineAsyncOpHandle<FCommerceQueryOffers> FCommerceEOS::QueryOffers(FCommerceQu
 		Op->Then([this](TOnlineAsyncOp<FCommerceQueryOffers>& Op)
 		{
 			const FCommerceQueryOffers::Params& Params = Op.GetParams();
-			if (!Services.Get<FAuthEOS>()->IsLoggedIn(Params.LocalAccountId))
+			if (!Services.Get<IAuth>()->IsLoggedIn(Params.LocalAccountId))
 			{
 				Op.SetError(Errors::NotLoggedIn());
 				return;
@@ -154,7 +154,7 @@ TOnlineAsyncOpHandle<FCommerceQueryOffersById> FCommerceEOS::QueryOffersById(FCo
 
 TOnlineResult<FCommerceGetOffers> FCommerceEOS::GetOffers(FCommerceGetOffers::Params&& Params)
 {
-	if (!Services.Get<FAuthEOS>()->IsLoggedIn(Params.LocalAccountId))
+	if (!Services.Get<IAuth>()->IsLoggedIn(Params.LocalAccountId))
 	{
 		return TOnlineResult<FCommerceGetOffers>(Errors::NotLoggedIn());
 	}
@@ -168,7 +168,7 @@ TOnlineResult<FCommerceGetOffers> FCommerceEOS::GetOffers(FCommerceGetOffers::Pa
 
 TOnlineResult<FCommerceGetOffersById> FCommerceEOS::GetOffersById(FCommerceGetOffersById::Params&& Params)
 {
-	if (!Services.Get<FAuthEOS>()->IsLoggedIn(Params.LocalAccountId))
+	if (!Services.Get<IAuth>()->IsLoggedIn(Params.LocalAccountId))
 	{
 		return TOnlineResult<FCommerceGetOffersById>(Errors::NotLoggedIn());
 	}
@@ -200,7 +200,7 @@ TOnlineAsyncOpHandle<FCommerceCheckout> FCommerceEOS::Checkout(FCommerceCheckout
 	Op->Then([this](TOnlineAsyncOp<FCommerceCheckout>& Op, TPromise<const EOS_Ecom_CheckoutCallbackInfo*>&& Promise)
 	{
 		const FCommerceCheckout::Params& Params = Op.GetParams();
-		if (!Services.Get<FAuthEOS>()->IsLoggedIn(Params.LocalAccountId))
+		if (!Services.Get<IAuth>()->IsLoggedIn(Params.LocalAccountId))
 		{
 			Op.SetError(Errors::NotLoggedIn());
 			Promise.SetValue(nullptr);
@@ -280,7 +280,7 @@ TOnlineAsyncOpHandle<FCommerceQueryTransactionEntitlements> FCommerceEOS::QueryT
 		Op->Then([this](TOnlineAsyncOp<FCommerceQueryTransactionEntitlements>& Op)
 		{
 			const FCommerceQueryTransactionEntitlements::Params& Params = Op.GetParams();
-			if (!Services.Get<FAuthEOS>()->IsLoggedIn(Params.LocalAccountId))
+			if (!Services.Get<IAuth>()->IsLoggedIn(Params.LocalAccountId))
 			{
 				Op.SetError(Errors::NotLoggedIn());
 				return;
@@ -355,7 +355,7 @@ TOnlineAsyncOpHandle<FCommerceQueryEntitlements> FCommerceEOS::QueryEntitlements
 		Op->Then([this](TOnlineAsyncOp<FCommerceQueryEntitlements>& Op)
 		{
 			const FCommerceQueryEntitlements::Params& Params = Op.GetParams();
-			if (!Services.Get<FAuthEOS>()->IsLoggedIn(Params.LocalAccountId))
+			if (!Services.Get<IAuth>()->IsLoggedIn(Params.LocalAccountId))
 			{
 				Op.SetError(Errors::NotLoggedIn());
 				return;
@@ -495,7 +495,7 @@ TOnlineAsyncOpHandle<FCommerceQueryEntitlements> FCommerceEOS::QueryEntitlements
 
 TOnlineResult<FCommerceGetEntitlements> FCommerceEOS::GetEntitlements(FCommerceGetEntitlements::Params&& Params)
 {
-	if (!Services.Get<FAuthEOS>()->IsLoggedIn(Params.LocalAccountId))
+	if (!Services.Get<IAuth>()->IsLoggedIn(Params.LocalAccountId))
 	{
 		return TOnlineResult<FCommerceGetEntitlements>(Errors::NotLoggedIn());
 	}
@@ -514,7 +514,7 @@ TOnlineAsyncOpHandle<FCommerceRedeemEntitlement> FCommerceEOS::RedeemEntitlement
 	Op->Then([this](TOnlineAsyncOp<FCommerceRedeemEntitlement>& Op, TPromise<const EOS_Ecom_RedeemEntitlementsCallbackInfo*>&& Promise)
 	{
 		const FCommerceRedeemEntitlement::Params& Params = Op.GetParams();
-		if (!Services.Get<FAuthEOS>()->IsLoggedIn(Params.LocalAccountId))
+		if (!Services.Get<IAuth>()->IsLoggedIn(Params.LocalAccountId))
 		{
 			Op.SetError(Errors::NotLoggedIn());
 			Promise.SetValue(nullptr);
@@ -569,7 +569,7 @@ TOnlineAsyncOpHandle<FCommerceRetrieveS2SToken> FCommerceEOS::RetrieveS2SToken(F
 		Op->Then([this](TOnlineAsyncOp<FCommerceRetrieveS2SToken>& Op)
 		{
 			const FCommerceRetrieveS2SToken::Params& Params = Op.GetParams();
-			if (!Services.Get<FAuthEOS>()->IsLoggedIn(Params.LocalAccountId))
+			if (!Services.Get<IAuth>()->IsLoggedIn(Params.LocalAccountId))
 			{
 				Op.SetError(Errors::NotLoggedIn());
 				return;
diff --git a/Engine/Plugins/Online/OnlineServicesEOS/Source/Private/Online/EpicAccountIdResolverEOS.cpp b/Engine/Plugins/Online/OnlineServicesEOS/Source/Private/Online/EpicAccountIdResolverEOS.cpp
new file mode 100644
--- /dev/null
+++ b/Engine/Plugins/Online/OnlineServicesEOS/Source/Private/Online/EpicAccountIdResolverEOS.cpp
@@ -0,0 +1,38 @@
+// Copyright Epic Games, Inc. All Rights Reserved.
+
+#include "Online/EpicAccountIdResolverEOS.h"
+
+#include "Online/AuthEOS.h"
+#include "Online/OnlineServicesEOS.h"
+
+namespace UE::Online {
+
+FEpicAccountIdResolverEOS::FEpicAccountIdResolverEOS(FOnlineServicesEOS& InServices)
+: TOnlineComponent(TEXT("EpicAccountIdResolver"), InServices)
+{
+}
+
+TFuture<FAccountId> FEpicAccountIdResolverEOS::ResolveAccountId(const FAccountId LocalAccountId, const EOS_EpicAccountId EpicAccountId)
+{
+	TPromise<FAccountId> Promise;
+	TFuture<FAccountId> Future = Promise.GetFuture();
+
+	Services.Get<FAuthEOS>()->ResolveAccountIds(LocalAccountId, { EpicAccountId }).Next([Promise = MoveTemp(Promise)](const TArray<FAccountId>& AccountIds) mutable
+	{
+		FAccountId Result;
+		if (AccountIds.Num() == 1)
+		{
+			Result = AccountIds[0];
+		}
+		Promise.SetValue(Result);
+	});
+
+	return Future;
+}
+
+TFuture<TArray<FAccountId>> FEpicAccountIdResolverEOS::ResolveAccountIds(const FAccountId LocalAccountId, const TArray<EOS_EpicAccountId>& EpicAccountIds)
+{
+	return Services.Get<FAuthEOS>()->ResolveAccountIds(LocalAccountId, EpicAccountIds);
+}
+
+/* UE::Online */ }
diff --git a/Engine/Plugins/Online/OnlineServicesEOS/Source/Private/Online/EpicProductUserIdResolverEOS.cpp b/Engine/Plugins/Online/OnlineServicesEOS/Source/Private/Online/EpicProductUserIdResolverEOS.cpp
new file mode 100644
--- /dev/null
+++ b/Engine/Plugins/Online/OnlineServicesEOS/Source/Private/Online/EpicProductUserIdResolverEOS.cpp
@@ -0,0 +1,38 @@
+// Copyright Epic Games, Inc. All Rights Reserved.
+
+#include "Online/EpicProductUserIdResolverEOS.h"
+
+#include "Online/AuthEOS.h"
+#include "Online/OnlineServicesEOS.h"
+
+namespace UE::Online {
+
+FEpicProductUserIdResolverEOS::FEpicProductUserIdResolverEOS(FOnlineServicesEOS& InServices)
+: TOnlineComponent(TEXT("EpicProductUserIdResolver"), InServices)
+{
+}
+
+TFuture<FAccountId> FEpicProductUserIdResolverEOS::ResolveAccountId(const FAccountId LocalAccountId, const EOS_ProductUserId ProductUserId)
+{
+	TPromise<FAccountId> Promise;
+	TFuture<FAccountId> Future = Promise.GetFuture();
+
+	Services.Get<FAuthEOS>()->ResolveAccountIds(LocalAccountId, { ProductUserId }).Next([Promise = MoveTemp(Promise)](const TArray<FAccountId>& AccountIds) mutable
+	{
+		FAccountId Result;
+		if (AccountIds.Num() == 1)
+		{
+			Result = AccountIds[0];
+		}
+		Promise.SetValue(Result);
+	});
+
+	return Future;
+}
+
+TFuture<TArray<FAccountId>> FEpicProductUserIdResolverEOS::ResolveAccountIds(const FAccountId LocalAccountId, const TArray<EOS_ProductUserId>& ProductUserIds)
+{
+	return Services.Get<FAuthEOS>()->ResolveAccountIds(LocalAccountId, ProductUserIds);
+}
+
+/* UE::Online */ }
diff --git a/Engine/Plugins/Online/OnlineServicesEOS/Source/Private/Online/OnlineServicesEOS.cpp b/Engine/Plugins/Online/OnlineServicesEOS/Source/Private/Online/OnlineServicesEOS.cpp
--- a/Engine/Plugins/Online/OnlineServicesEOS/Source/Private/Online/OnlineServicesEOS.cpp
+++ b/Engine/Plugins/Online/OnlineServicesEOS/Source/Private/Online/OnlineServicesEOS.cpp
@@ -3,6 +3,8 @@
 #include "Online/OnlineServicesEOS.h"
 
 #include "Online/AuthEOS.h"
+#include "Online/EpicAccountIdResolverEOS.h"
+#include "Online/EpicProductUserIdResolverEOS.h"
 #include "Online/ExternalUIEOS.h"
 #include "Online/SocialEOS.h"
 #include "Online/PresenceEOS.h"
@@ -19,6 +21,8 @@ FOnlineServicesEOS::FOnlineServicesEOS(FName InInstanceName, FName InInstanceCon
 void FOnlineServicesEOS::RegisterComponents()
 {
 	Components.Register<FAuthEOS>(*this);
+	Components.Register<FEpicAccountIdResolverEOS>(*this);
+	Components.Register<FEpicProductUserIdResolverEOS>(*this);
 	Components.Register<FExternalUIEOS>(*this);
 	Components.Register<FSocialEOS>(*this);
 	Components.Register<FPresenceEOS>(*this);
diff --git a/Engine/Plugins/Online/OnlineServicesEOS/Source/Private/Online/PresenceEOS.cpp b/Engine/Plugins/Online/OnlineServicesEOS/Source/Private/Online/PresenceEOS.cpp
--- a/Engine/Plugins/Online/OnlineServicesEOS/Source/Private/Online/PresenceEOS.cpp
+++ b/Engine/Plugins/Online/OnlineServicesEOS/Source/Private/Online/PresenceEOS.cpp
@@ -5,6 +5,7 @@
 #include "EOSShared.h"
 #include "IEOSSDKManager.h"
 #include "Online/AuthEOS.h"
+#include "Online/EpicAccountIdResolver.h"
 #include "Online/OnlineErrorEpicCommon.h"
 #include "Online/OnlineIdEOS.h"
 #include "Online/OnlineServicesEOS.h"
@@ -106,7 +107,7 @@ void FPresenceEOS::HandleAuthLoginStatusChanged(const FAuthLoginStatusChanged& E
 		{
 			for (EOS_EpicAccountId& PresenceAccountId : *PendingPresenceUpdateArray)
 			{
-				Services.Get<FAuthEOS>()->ResolveAccountId(LocalAccountId, PresenceAccountId)
+				Services.Get<IEpicAccountIdResolver>()->ResolveAccountId(LocalAccountId, PresenceAccountId)
 					.Next([WeakThis = AsWeak(), LocalAccountId](const FAccountId& PresenceAccountId)
 						{
 							if (TSharedPtr<FPresenceEOS> StrongThis = StaticCastSharedPtr<FPresenceEOS>(WeakThis.Pin()))
@@ -129,7 +130,7 @@ TOnlineAsyncOpHandle<FQueryPresence> FPresenceEOS::QueryPresence(FQueryPresence:
 	{
 		// Initialize
 		const FQueryPresence::Params& Params = Op->GetParams();
-		if (!Services.Get<FAuthEOS>()->IsLoggedIn(Params.LocalAccountId))
+		if (!Services.Get<IAuth>()->IsLoggedIn(Params.LocalAccountId))
 		{
 			Op->SetError(Errors::NotLoggedIn());
 			return Op->GetHandle();
@@ -142,7 +143,7 @@ TOnlineAsyncOpHandle<FQueryPresence> FPresenceEOS::QueryPresence(FQueryPresence:
 		}
 
 		// TODO:  If we try to query a local user's presence, is that an error, should we return the cached state, should we still ask EOS?
-		const bool bIsLocalUser = Services.Get<FAuthEOS>()->IsLoggedIn(Params.TargetAccountId);
+		const bool bIsLocalUser = Services.Get<IAuth>()->IsLoggedIn(Params.TargetAccountId);
 		if (bIsLocalUser)
 		{
 			Op->SetError(Errors::CannotQueryLocalUsers()); 
@@ -200,7 +201,7 @@ TOnlineAsyncOpHandle<FUpdatePresence> FPresenceEOS::UpdatePresence(FUpdatePresen
 {
 	TOnlineAsyncOpRef<FUpdatePresence> Op = GetOp<FUpdatePresence>(MoveTemp(InParams));
 	const FUpdatePresence::Params& Params = Op->GetParams();
-	if (!Services.Get<FAuthEOS>()->IsLoggedIn(Params.LocalAccountId))
+	if (!Services.Get<IAuth>()->IsLoggedIn(Params.LocalAccountId))
 	{
 		Op->SetError(Errors::NotLoggedIn());
 		return Op->GetHandle();
@@ -390,7 +391,7 @@ TOnlineAsyncOpHandle<FPartialUpdatePresence> FPresenceEOS::PartialUpdatePresence
 	{
 		// Initialize
 		const FPartialUpdatePresence::Params& Params = Op->GetParams();
-		if (!Services.Get<FAuthEOS>()->IsLoggedIn(Params.LocalAccountId))
+		if (!Services.Get<IAuth>()->IsLoggedIn(Params.LocalAccountId))
 		{
 			Op->SetError(Errors::NotLoggedIn());
 			return Op->GetHandle();
@@ -707,7 +708,7 @@ void FPresenceEOS::HandlePresenceChanged(const EOS_Presence_PresenceChangedCallb
 {
 	if (const FAccountId LocalAccountId = FindAccountId(Data->LocalUserId))
 	{
-		Services.Get<FAuthEOS>()->ResolveAccountId(LocalAccountId, Data->PresenceUserId)
+		Services.Get<IEpicAccountIdResolver>()->ResolveAccountId(LocalAccountId, Data->PresenceUserId)
 		.Next([WeakThis = AsWeak(), LocalAccountId, Func = __FUNCTION__](const FAccountId& PresenceAccountId)
 			{
 				if (TSharedPtr<FPresenceEOS> StrongThis = StaticCastSharedPtr<FPresenceEOS>(WeakThis.Pin()))
diff --git a/Engine/Plugins/Online/OnlineServicesEOS/Source/Private/Online/SocialEOS.cpp b/Engine/Plugins/Online/OnlineServicesEOS/Source/Private/Online/SocialEOS.cpp
--- a/Engine/Plugins/Online/OnlineServicesEOS/Source/Private/Online/SocialEOS.cpp
+++ b/Engine/Plugins/Online/OnlineServicesEOS/Source/Private/Online/SocialEOS.cpp
@@ -4,6 +4,7 @@
 
 #include "IEOSSDKManager.h"
 #include "Online/AuthEOS.h"
+#include "Online/EpicAccountIdResolver.h"
 #include "Online/OnlineErrorEpicCommon.h"
 #include "Online/OnlineIdEOS.h"
 #include "Online/OnlineServicesEOS.h"
@@ -75,7 +76,7 @@ TOnlineAsyncOpHandle<FQueryFriends> FSocialEOS::QueryFriends(FQueryFriends::Para
 			return;
 		}
 
-		if (!Services.Get<FAuthEOS>()->IsLoggedIn(Params.LocalAccountId))
+		if (!Services.Get<IAuth>()->IsLoggedIn(Params.LocalAccountId))
 		{
 			Op.SetError(Errors::NotLoggedIn());
 			return;
@@ -120,7 +121,7 @@ TOnlineAsyncOpHandle<FQueryFriends> FSocialEOS::QueryFriends(FQueryFriends::Para
 		}
 		return CurrentFriends;
 	})
-	.Then(Services.Get<FAuthEOS>()->ResolveEpicIdsFn())
+	.Then(Services.Get<IEpicAccountIdResolver>()->ResolveEpicAccountIdsFn())
 	.Then([this](TOnlineAsyncOp<FQueryFriends>& Op, const TArray<FAccountId>& CurrentFriendIds)
 	{
 		const FAccountId LocalUserId = Op.GetParams().LocalAccountId;
@@ -199,7 +200,7 @@ TOnlineAsyncOpHandle<FSendFriendInvite> FSocialEOS::SendFriendInvite(FSendFriend
 			return;
 		}
 
-		if (!Services.Get<FAuthEOS>()->IsLoggedIn(Params.LocalAccountId))
+		if (!Services.Get<IAuth>()->IsLoggedIn(Params.LocalAccountId))
 		{
 			Op.SetError(Errors::NotLoggedIn());
 			return;
@@ -251,7 +252,7 @@ TOnlineAsyncOpHandle<FAcceptFriendInvite> FSocialEOS::AcceptFriendInvite(FAccept
 			return;
 		}
 
-		if (!Services.Get<FAuthEOS>()->IsLoggedIn(Params.LocalAccountId))
+		if (!Services.Get<IAuth>()->IsLoggedIn(Params.LocalAccountId))
 		{
 			Op.SetError(Errors::NotLoggedIn());
 			return;
@@ -303,7 +304,7 @@ TOnlineAsyncOpHandle<FRejectFriendInvite> FSocialEOS::RejectFriendInvite(FReject
 			return;
 		}
 
-		if (!Services.Get<FAuthEOS>()->IsLoggedIn(Params.LocalAccountId))
+		if (!Services.Get<IAuth>()->IsLoggedIn(Params.LocalAccountId))
 		{
 			Op.SetError(Errors::NotLoggedIn());
 			return;
@@ -361,7 +362,7 @@ TOnlineAsyncOpHandle<FBlockUser> FSocialEOS::BlockUser(FBlockUser::Params&& Para
 			return;
 		}
 
-		if (!Services.Get<FAuthEOS>()->IsLoggedIn(Params.LocalAccountId))
+		if (!Services.Get<IAuth>()->IsLoggedIn(Params.LocalAccountId))
 		{
 			Op.SetError(Errors::NotLoggedIn());
 			return;
@@ -404,7 +405,7 @@ void FSocialEOS::HandleFriendsUpdate(const EOS_Friends_OnFriendsUpdateInfo* Data
 		return;
 	}
 
-	Services.Get<FAuthEOS>()->ResolveAccountId(LocalAccountId, Data->TargetUserId)
+	Services.Get<IEpicAccountIdResolver>()->ResolveAccountId(LocalAccountId, Data->TargetUserId)
 	.Next([this, LocalAccountId, PreviousStatus = Data->PreviousStatus, CurrentStatus = Data->CurrentStatus](const FAccountId& FriendAccountId)
 	{
 		const EOS_EpicAccountId LocalUserEasId = GetEpicAccountIdChecked(LocalAccountId);
diff --git a/Engine/Plugins/Online/OnlineServicesEOS/Source/Private/Online/UserInfoEOS.cpp b/Engine/Plugins/Online/OnlineServicesEOS/Source/Private/Online/UserInfoEOS.cpp
--- a/Engine/Plugins/Online/OnlineServicesEOS/Source/Private/Online/UserInfoEOS.cpp
+++ b/Engine/Plugins/Online/OnlineServicesEOS/Source/Private/Online/UserInfoEOS.cpp
@@ -47,7 +47,7 @@ TOnlineAsyncOpHandle<FQueryUserInfo> FUserInfoEOS::QueryUserInfo(FQueryUserInfo:
 			{
 				const FQueryUserInfo::Params& Params = Op.GetParams();
 
-				if (!Services.Get<FAuthEOS>()->IsLoggedIn(Params.LocalAccountId))
+				if (!Services.Get<IAuth>()->IsLoggedIn(Params.LocalAccountId))
 				{
 					Op.SetError(Errors::NotLoggedIn());
 					Promise.EmplaceValue();
@@ -121,7 +121,7 @@ TOnlineAsyncOpHandle<FQueryUserInfo> FUserInfoEOS::QueryUserInfo(FQueryUserInfo:
 
 TOnlineResult<FGetUserInfo> FUserInfoEOS::GetUserInfo(FGetUserInfo::Params&& Params)
 {
-	if (!Services.Get<FAuthEOS>()->IsLoggedIn(Params.LocalAccountId))
+	if (!Services.Get<IAuth>()->IsLoggedIn(Params.LocalAccountId))
 	{
 		return TOnlineResult<FGetUserInfo>(Errors::NotLoggedIn());
 	}
diff --git a/Engine/Plugins/Online/OnlineServicesEOS/Source/Public/Online/AuthEOS.h b/Engine/Plugins/Online/OnlineServicesEOS/Source/Public/Online/AuthEOS.h
--- a/Engine/Plugins/Online/OnlineServicesEOS/Source/Public/Online/AuthEOS.h
+++ b/Engine/Plugins/Online/OnlineServicesEOS/Source/Public/Online/AuthEOS.h
@@ -38,17 +38,8 @@ public:
 	UE_API virtual TOnlineResult<FAuthGetLinkAccountContinuationId> GetLinkAccountContinuationId(FAuthGetLinkAccountContinuationId::Params&& Params) const;
 	// End IAuth
 
-	// Begin FAuthEOSGS
-	UE_API virtual TFuture<FAccountId> ResolveAccountId(const FAccountId& LocalAccountId, const EOS_ProductUserId ProductUserId) override;
-	UE_API virtual TFuture<TArray<FAccountId>> ResolveAccountIds(const FAccountId& LocalAccountId, const TArray<EOS_ProductUserId>& ProductUserIds) override;
-	UE_API virtual TFunction<TFuture<FAccountId>(FOnlineAsyncOp& InAsyncOp, const EOS_ProductUserId& ProductUserId)> ResolveProductIdFn() override;
-	UE_API virtual TFunction<TFuture<TArray<FAccountId>>(FOnlineAsyncOp& InAsyncOp, const TArray<EOS_ProductUserId>& ProductUserIds)> ResolveProductIdsFn() override;
-	// End FAuthEOSGS
-
-	UE_API TFuture<FAccountId> ResolveAccountId(const FAccountId& LocalAccountId, const EOS_EpicAccountId EpicAccountId);
 	UE_API TFuture<TArray<FAccountId>> ResolveAccountIds(const FAccountId& LocalAccountId, const TArray<EOS_EpicAccountId>& EpicAccountIds);
-	UE_API TFunction<TFuture<FAccountId>(FOnlineAsyncOp& InAsyncOp, const EOS_EpicAccountId& EpicAccountId)> ResolveEpicIdFn();
-	UE_API TFunction<TFuture<TArray<FAccountId>>(FOnlineAsyncOp& InAsyncOp, const TArray<EOS_EpicAccountId>& EpicAccountIds)> ResolveEpicIdsFn();
+	UE_API TFuture<TArray<FAccountId>> ResolveAccountIds(const FAccountId& LocalAccountId, const TArray<EOS_ProductUserId>& ProductUserIds);
 
 protected:
 	UE_API void ProcessSuccessfulLogin(TOnlineAsyncOp<FAuthLogin>& InAsyncOp);
diff --git a/Engine/Plugins/Online/OnlineServicesEOS/Source/Public/Online/EpicAccountIdResolverEOS.h b/Engine/Plugins/Online/OnlineServicesEOS/Source/Public/Online/EpicAccountIdResolverEOS.h
new file mode 100644
--- /dev/null
+++ b/Engine/Plugins/Online/OnlineServicesEOS/Source/Public/Online/EpicAccountIdResolverEOS.h
@@ -0,0 +1,21 @@
+// Copyright Epic Games, Inc. All Rights Reserved.
+
+#pragma once
+
+#include "Online/EpicAccountIdResolver.h"
+#include "Online/OnlineComponent.h"
+
+namespace UE::Online {
+
+class FEpicAccountIdResolverEOS : public TOnlineComponent<IEpicAccountIdResolver>
+{
+public:
+	FEpicAccountIdResolverEOS(class FOnlineServicesEOS& InServices);
+
+	// Begin IEpicAccountIdResolver
+	virtual TFuture<FAccountId> ResolveAccountId(const FAccountId LocalAccountId, const EOS_EpicAccountId EpicAccountId) override;
+	virtual TFuture<TArray<FAccountId>> ResolveAccountIds(const FAccountId LocalAccountId, const TArray<EOS_EpicAccountId>& EpicAccountIds) override;
+	// End IEpicAccountIdResolver
+};
+
+/* UE::Online */ }
diff --git a/Engine/Plugins/Online/OnlineServicesEOS/Source/Public/Online/EpicProductUserIdResolverEOS.h b/Engine/Plugins/Online/OnlineServicesEOS/Source/Public/Online/EpicProductUserIdResolverEOS.h
new file mode 100644
--- /dev/null
+++ b/Engine/Plugins/Online/OnlineServicesEOS/Source/Public/Online/EpicProductUserIdResolverEOS.h
@@ -0,0 +1,21 @@
+// Copyright Epic Games, Inc. All Rights Reserved.
+
+#pragma once
+
+#include "Online/EpicProductUserIdResolver.h"
+#include "Online/OnlineComponent.h"
+
+namespace UE::Online {
+
+class FEpicProductUserIdResolverEOS : public TOnlineComponent<IEpicProductUserIdResolver>
+{
+public:
+	FEpicProductUserIdResolverEOS(class FOnlineServicesEOS& InServices);
+
+	// Begin IEpicProductUserIdResolver
+	virtual TFuture<FAccountId> ResolveAccountId(const FAccountId LocalAccountId, const EOS_ProductUserId ProductUserId) override;
+	virtual TFuture<TArray<FAccountId>> ResolveAccountIds(const FAccountId LocalAccountId, const TArray<EOS_ProductUserId>& ProductUserIds) override;
+	// End IEpicProductUserIdResolver
+};
+
+/* UE::Online */ }
diff --git a/Engine/Plugins/Online/OnlineServicesEOSGS/Source/Private/Online/AchievementsEOSGS.cpp b/Engine/Plugins/Online/OnlineServicesEOSGS/Source/Private/Online/AchievementsEOSGS.cpp
--- a/Engine/Plugins/Online/OnlineServicesEOSGS/Source/Private/Online/AchievementsEOSGS.cpp
+++ b/Engine/Plugins/Online/OnlineServicesEOSGS/Source/Private/Online/AchievementsEOSGS.cpp
@@ -52,7 +52,7 @@ TOnlineAsyncOpHandle<FQueryAchievementDefinitions> FAchievementsEOSGS::QueryAchi
 		{
 			const FQueryAchievementDefinitions::Params& Params = InAsyncOp.GetParams();
 
-			if (!Services.Get<FAuthEOSGS>()->IsLoggedIn(Params.LocalAccountId))
+			if (!Services.Get<IAuth>()->IsLoggedIn(Params.LocalAccountId))
 			{
 				InAsyncOp.SetError(Errors::InvalidUser());
 			}
@@ -148,7 +148,7 @@ TOnlineAsyncOpHandle<FQueryAchievementDefinitions> FAchievementsEOSGS::QueryAchi
 
 TOnlineResult<FGetAchievementIds> FAchievementsEOSGS::GetAchievementIds(FGetAchievementIds::Params&& Params)
 {
-	if (!Services.Get<FAuthEOSGS>()->IsLoggedIn(Params.LocalAccountId))
+	if (!Services.Get<IAuth>()->IsLoggedIn(Params.LocalAccountId))
 	{
 		return TOnlineResult<FGetAchievementIds>(Errors::InvalidUser());
 	}
@@ -166,7 +166,7 @@ TOnlineResult<FGetAchievementIds> FAchievementsEOSGS::GetAchievementIds(FGetAchi
 
 TOnlineResult<FGetAchievementDefinition> FAchievementsEOSGS::GetAchievementDefinition(FGetAchievementDefinition::Params&& Params)
 {
-	if (!Services.Get<FAuthEOSGS>()->IsLoggedIn(Params.LocalAccountId))
+	if (!Services.Get<IAuth>()->IsLoggedIn(Params.LocalAccountId))
 	{
 		return TOnlineResult<FGetAchievementDefinition>(Errors::InvalidUser());
 	}
@@ -194,7 +194,7 @@ TOnlineAsyncOpHandle<FQueryAchievementStates> FAchievementsEOSGS::QueryAchieveme
 		Op->Then([this](TOnlineAsyncOp<FQueryAchievementStates>& InAsyncOp)
 		{
 			const FQueryAchievementStates::Params& Params = InAsyncOp.GetParams();
-			if (!Services.Get<FAuthEOSGS>()->IsLoggedIn(Params.LocalAccountId))
+			if (!Services.Get<IAuth>()->IsLoggedIn(Params.LocalAccountId))
 			{
 				InAsyncOp.SetError(Errors::InvalidUser());
 			}
@@ -283,7 +283,7 @@ TOnlineAsyncOpHandle<FQueryAchievementStates> FAchievementsEOSGS::QueryAchieveme
 
 TOnlineResult<FGetAchievementState> FAchievementsEOSGS::GetAchievementState(FGetAchievementState::Params&& Params) const
 {
-	if (!Services.Get<FAuthEOSGS>()->IsLoggedIn(Params.LocalAccountId))
+	if (!Services.Get<IAuth>()->IsLoggedIn(Params.LocalAccountId))
 	{
 		return TOnlineResult<FGetAchievementState>(Errors::InvalidUser());
 	}
@@ -312,7 +312,7 @@ TOnlineAsyncOpHandle<FUnlockAchievements> FAchievementsEOSGS::UnlockAchievements
 		Op->Then([this](TOnlineAsyncOp<FUnlockAchievements>& InAsyncOp)
 		{
 			const FUnlockAchievements::Params& Params = InAsyncOp.GetParams();
-			if (!Services.Get<FAuthEOSGS>()->IsLoggedIn(Params.LocalAccountId))
+			if (!Services.Get<IAuth>()->IsLoggedIn(Params.LocalAccountId))
 			{
 				InAsyncOp.SetError(Errors::InvalidUser());
 				return;
diff --git a/Engine/Plugins/Online/OnlineServicesEOSGS/Source/Private/Online/AuthEOSGS.cpp b/Engine/Plugins/Online/OnlineServicesEOSGS/Source/Private/Online/AuthEOSGS.cpp
--- a/Engine/Plugins/Online/OnlineServicesEOSGS/Source/Private/Online/AuthEOSGS.cpp
+++ b/Engine/Plugins/Online/OnlineServicesEOSGS/Source/Private/Online/AuthEOSGS.cpp
@@ -860,48 +860,6 @@ TOnlineAsyncOpHandle<FAuthEndVerifiedAuthSession> FAuthEOSGS::EndVerifiedAuthSes
 	return Operation->GetHandle();
 }
 
-TFuture<FAccountId> FAuthEOSGS::ResolveAccountId(const FAccountId& LocalAccountId, const EOS_ProductUserId ProductUserId)
-{
-	return MakeFulfilledPromise<FAccountId>(CreateAccountId(ProductUserId)).GetFuture();
-}
-
-TFuture<TArray<FAccountId>> FAuthEOSGS::ResolveAccountIds(const FAccountId& LocalAccountId, const TArray<EOS_ProductUserId>& InProductUserIds)
-{
-	TArray<FAccountId> AccountIds;
-	AccountIds.Reserve(InProductUserIds.Num());
-	for (const EOS_ProductUserId ProductUserId : InProductUserIds)
-	{
-		AccountIds.Emplace(CreateAccountId(ProductUserId));
-	}
-	return MakeFulfilledPromise<TArray<FAccountId>>(MoveTemp(AccountIds)).GetFuture();
-}
-
-TFunction<TFuture<FAccountId>(FOnlineAsyncOp& InAsyncOp, const EOS_ProductUserId& ProductUserId)> FAuthEOSGS::ResolveProductIdFn()
-{
-	return [this](FOnlineAsyncOp& InAsyncOp, const EOS_ProductUserId& ProductUserId)
-	{
-		const FAccountId* LocalAccountIdPtr = InAsyncOp.Data.Get<FAccountId>(TEXT("LocalAccountId"));
-		if (!ensure(LocalAccountIdPtr))
-		{
-			return MakeFulfilledPromise<FAccountId>().GetFuture();
-		}
-		return ResolveAccountId(*LocalAccountIdPtr, ProductUserId);
-	};
-}
-
-TFunction<TFuture<TArray<FAccountId>>(FOnlineAsyncOp& InAsyncOp, const TArray<EOS_ProductUserId>& ProductUserIds)> FAuthEOSGS::ResolveProductIdsFn()
-{
-	return [this](FOnlineAsyncOp& InAsyncOp, const TArray<EOS_ProductUserId>& ProductUserIds)
-	{
-		const FAccountId* LocalAccountIdPtr = InAsyncOp.Data.Get<FAccountId>(TEXT("LocalAccountId"));
-		if (!ensure(LocalAccountIdPtr))
-		{
-			return MakeFulfilledPromise<TArray<FAccountId>>().GetFuture();
-		}
-		return ResolveAccountIds(*LocalAccountIdPtr, ProductUserIds);
-	};
-}
-
 TFuture<TDefaultErrorResult<FAuthLoginEASImpl>> FAuthEOSGS::LoginEASImpl(const FAuthLoginEASImpl::Params& LoginParams)
 {
 	TDefaultErrorResultInternal<FEOSAuthLoginOptions> LoginOptionsResult = FEOSAuthLoginOptions::Create(LoginParams.CredentialsType, LoginParams.CredentialsId, LoginParams.CredentialsToken, LoginParams.Scopes);
diff --git a/Engine/Plugins/Online/OnlineServicesEOSGS/Source/Private/Online/EpicProductUserIdResolverEOSGS.cpp b/Engine/Plugins/Online/OnlineServicesEOSGS/Source/Private/Online/EpicProductUserIdResolverEOSGS.cpp
new file mode 100644
--- /dev/null
+++ b/Engine/Plugins/Online/OnlineServicesEOSGS/Source/Private/Online/EpicProductUserIdResolverEOSGS.cpp
@@ -0,0 +1,34 @@
+// Copyright Epic Games, Inc. All Rights Reserved.
+
+#include "Online/EpicProductUserIdResolverEOSGS.h"
+
+#include "Online/OnlineIdEOSGS.h"
+#include "Online/OnlineServicesEOSGS.h"
+
+namespace UE::Online {
+
+FEpicProductUserIdResolverEOSGS::FEpicProductUserIdResolverEOSGS(FOnlineServicesEOSGS& InServices)
+: TOnlineComponent(TEXT("EpicProductUserIdResolver"), InServices)
+{
+}
+
+TFuture<FAccountId> FEpicProductUserIdResolverEOSGS::ResolveAccountId(const FAccountId LocalAccountId, const EOS_ProductUserId ProductUserId)
+{
+	const FAccountId AccountId = FOnlineAccountIdRegistryEOSGS::GetRegistered(EOnlineServices::Epic).FindOrAddAccountId(ProductUserId);
+	
+	return MakeFulfilledPromise<FAccountId>(AccountId).GetFuture();
+}
+
+TFuture<TArray<FAccountId>> FEpicProductUserIdResolverEOSGS::ResolveAccountIds(const FAccountId LocalAccountId, const TArray<EOS_ProductUserId>& ProductUserIds)
+{
+	TArray<FAccountId> AccountIds;
+	AccountIds.Reserve(ProductUserIds.Num());
+	for (const EOS_ProductUserId ProductUserId : ProductUserIds)
+	{
+		const FAccountId AccountId = FOnlineAccountIdRegistryEOSGS::GetRegistered(EOnlineServices::Epic).FindOrAddAccountId(ProductUserId);
+		AccountIds.Emplace(AccountId);
+	}
+	return MakeFulfilledPromise<TArray<FAccountId>>(MoveTemp(AccountIds)).GetFuture();
+}
+
+/* UE::Online */ }
diff --git a/Engine/Plugins/Online/OnlineServicesEOSGS/Source/Private/Online/LeaderboardsEOSGS.cpp b/Engine/Plugins/Online/OnlineServicesEOSGS/Source/Private/Online/LeaderboardsEOSGS.cpp
--- a/Engine/Plugins/Online/OnlineServicesEOSGS/Source/Private/Online/LeaderboardsEOSGS.cpp
+++ b/Engine/Plugins/Online/OnlineServicesEOSGS/Source/Private/Online/LeaderboardsEOSGS.cpp
@@ -37,7 +37,7 @@ TOnlineAsyncOpHandle<FReadEntriesForUsers> FLeaderboardsEOSGS::ReadEntriesForUse
 		Op->Then([this](TOnlineAsyncOp<FReadEntriesForUsers>& InAsyncOp)
 		{
 			const FReadEntriesForUsers::Params& Params = InAsyncOp.GetParams();
-			if (!Services.Get<FAuthEOSGS>()->IsLoggedIn(Params.LocalAccountId))
+			if (!Services.Get<IAuth>()->IsLoggedIn(Params.LocalAccountId))
 			{
 				InAsyncOp.SetError(Errors::InvalidUser());
 			}
@@ -175,7 +175,7 @@ TOnlineAsyncOpHandle<FReadEntriesAroundRank> FLeaderboardsEOSGS::ReadEntriesArou
 		Op->Then([this](TOnlineAsyncOp<FReadEntriesAroundRank>& InAsyncOp)
 		{
 			const FReadEntriesAroundRank::Params& Params = InAsyncOp.GetParams();
-			if (!Services.Get<FAuthEOSGS>()->IsLoggedIn(Params.LocalAccountId))
+			if (!Services.Get<IAuth>()->IsLoggedIn(Params.LocalAccountId))
 			{
 				InAsyncOp.SetError(Errors::InvalidUser());
 			}
@@ -220,7 +220,7 @@ TOnlineAsyncOpHandle<FReadEntriesAroundUser> FLeaderboardsEOSGS::ReadEntriesArou
 		Op->Then([this](TOnlineAsyncOp<FReadEntriesAroundUser>& InAsyncOp)
 		{
 			const FReadEntriesAroundUser::Params& Params = InAsyncOp.GetParams();
-			if (!Services.Get<FAuthEOSGS>()->IsLoggedIn(Params.LocalAccountId))
+			if (!Services.Get<IAuth>()->IsLoggedIn(Params.LocalAccountId))
 			{
 				InAsyncOp.SetError(Errors::InvalidUser());
 			}
diff --git a/Engine/Plugins/Online/OnlineServicesEOSGS/Source/Private/Online/LobbiesEOSGS.cpp b/Engine/Plugins/Online/OnlineServicesEOSGS/Source/Private/Online/LobbiesEOSGS.cpp
--- a/Engine/Plugins/Online/OnlineServicesEOSGS/Source/Private/Online/LobbiesEOSGS.cpp
+++ b/Engine/Plugins/Online/OnlineServicesEOSGS/Source/Private/Online/LobbiesEOSGS.cpp
@@ -39,11 +39,11 @@ void FLobbiesEOSGS::Initialize()
 	EOS_HLobby LobbyInterfaceHandle = EOS_Platform_GetLobbyInterface(*GetServices<FOnlineServicesEpicCommon>().GetEOSPlatformHandle());
 	check(LobbyInterfaceHandle != nullptr);
 
-	LobbyPrerequisites = MakeShared<FLobbyPrerequisitesEOS>(FLobbyPrerequisitesEOS{
+	LobbyPrerequisites = MakeShared<FLobbyPrerequisitesEOS>(
 		LobbyInterfaceHandle,
-		StaticCastSharedPtr<FAuthEOSGS>(Services.GetAuthInterface()),
-		{SDKManager->GetProductName(), GetBuildUniqueId()},
-		SchemaRegistry});
+		Services.Get<IEpicProductUserIdResolver>(),
+		FLobbyBucketIdEOS(SDKManager->GetProductName(), GetBuildUniqueId()),
+		SchemaRegistry);
 
 	LobbyDataRegistry = MakeShared<FLobbyDataRegistryEOS>(LobbyPrerequisites.ToSharedRef());
 
diff --git a/Engine/Plugins/Online/OnlineServicesEOSGS/Source/Private/Online/LobbiesEOSGSTypes.cpp b/Engine/Plugins/Online/OnlineServicesEOSGS/Source/Private/Online/LobbiesEOSGSTypes.cpp
--- a/Engine/Plugins/Online/OnlineServicesEOSGS/Source/Private/Online/LobbiesEOSGSTypes.cpp
+++ b/Engine/Plugins/Online/OnlineServicesEOSGS/Source/Private/Online/LobbiesEOSGSTypes.cpp
@@ -3,9 +3,12 @@
 #include "LobbiesEOSGSTypes.h"
 
 #include "EOSShared.h"
-#include "Online/AuthEOSGS.h"
+#include "EOSSharedTypes.h"
+#include "Online/EpicProductUserIdResolver.h"
+#include "Online/OnlineAsyncOp.h"
 #include "Online/OnlineErrorEpicCommon.h"
 #include "Online/OnlineIdEOSGS.h"
+#include "Online/OnlineServicesLog.h"
 #include "Online/OnlineUtils.h"
 #include <eos_lobby.h>
 
@@ -42,9 +45,9 @@ FString ToLogString(const FLobbyBucketIdEOS& BucketId)
 	return FString::Printf(TEXT("%s:%d"), *BucketId.GetProductName(), BucketId.GetProductVersion());
 }
 
-FLobbyBucketIdEOS::FLobbyBucketIdEOS(FString ProductName, int32 ProductVersion)
-	: ProductName(ProductName.Replace(UE_ONLINE_LOBBIES_BUCKET_ID_EOS_SEPARATOR, TEXT("_")))
-	, ProductVersion(ProductVersion)
+FLobbyBucketIdEOS::FLobbyBucketIdEOS(const FString& InProductName, int32 InProductVersion)
+	: ProductName(InProductName.Replace(UE_ONLINE_LOBBIES_BUCKET_ID_EOS_SEPARATOR, TEXT("_")))
+	, ProductVersion(InProductVersion)
 {
 }
 
@@ -302,10 +305,9 @@ bool FLobbyDetailsEOS::IsBucketCompatible() const
 
 TFuture<TDefaultErrorResultInternal<FLobbyServiceSnapshot>> FLobbyDetailsEOS::GetLobbySnapshot(bool bCopyMemberData) const
 {
-	TSharedPtr<FAuthEOSGS> AuthInterface = Prerequisites->AuthInterface.Pin();
-	if (!AuthInterface)
+	if (!Prerequisites->ResolverInterface)
 	{
-		UE_LOG(LogOnlineServices, Warning, TEXT("[FLobbyDetailsEOS::GetLobbySnapshot] Failed: Auth interface missing. Lobby[%s]"),
+		UE_LOG(LogOnlineServices, Warning, TEXT("[FLobbyDetailsEOS::GetLobbySnapshot] Failed: Resolver interface missing. Lobby[%s]"),
 			UTF8_TO_TCHAR(GetInfo()->GetLobbyIdEOS()));
 		return MakeFulfilledPromise<TDefaultErrorResultInternal<FLobbyServiceSnapshot>>(Errors::MissingInterface()).GetFuture();
 	}
@@ -357,7 +359,7 @@ TFuture<TDefaultErrorResultInternal<FLobbyServiceSnapshot>> FLobbyDetailsEOS::Ge
 	TFuture<TDefaultErrorResultInternal<FLobbyServiceSnapshot>> Future = Promise.GetFuture();
 
 	// Resolve lobby member product user ids to FAccountId before proceeding.
-	AuthInterface->ResolveAccountIds(AssociatedLocalUser, *MemberProductUserIds)
+	Prerequisites->ResolverInterface->ResolveAccountIds(AssociatedLocalUser, *MemberProductUserIds)
 	.Next([StrongThis = AsShared(), Promise = MoveTemp(Promise), MemberProductUserIds, LobbyOwner](TArray<FAccountId>&& ResolvedAccountIds) mutable
 	{
 		if (MemberProductUserIds->Num() != ResolvedAccountIds.Num())
diff --git a/Engine/Plugins/Online/OnlineServicesEOSGS/Source/Private/Online/LobbiesEOSGSTypes.h b/Engine/Plugins/Online/OnlineServicesEOSGS/Source/Private/Online/LobbiesEOSGSTypes.h
--- a/Engine/Plugins/Online/OnlineServicesEOSGS/Source/Private/Online/LobbiesEOSGSTypes.h
+++ b/Engine/Plugins/Online/OnlineServicesEOSGS/Source/Private/Online/LobbiesEOSGSTypes.h
@@ -25,8 +25,8 @@ class FLobbyDetailsInfoEOS;
 class FLobbyInviteDataEOSRegistry;
 class FLobbySchema;
 struct FLobbyEvents;
-class FAuthEOSGS;
 struct FLobbyMemberServiceSnapshot;
+class IEpicProductUserIdResolver;
 
 //--------------------------------------------------------------------------------------------------
 // Translated types
@@ -38,6 +38,8 @@ class FLobbyBucketIdEOS
 {
 public:
 	FLobbyBucketIdEOS() = default;
+	FLobbyBucketIdEOS(const FString& ProductName, int32 ProductVersion);
+
 	FLobbyBucketIdEOS(const FLobbyBucketIdEOS&) = default;
 	FLobbyBucketIdEOS(FLobbyBucketIdEOS&&) = default;
 	FLobbyBucketIdEOS& operator=(const FLobbyBucketIdEOS&) = default;
@@ -53,8 +55,6 @@ public:
 		return !(*this == Other);
 	}
 
-	FLobbyBucketIdEOS(FString ProductName, int32 ProductVersion);
-
 	const FString& GetProductName() const { return ProductName; }
 	int32 GetProductVersion() const { return ProductVersion; }
 	bool IsValid() const { return !ProductName.IsEmpty(); }
@@ -183,7 +183,7 @@ using FLobbySearchParameters = FFindLobbies::Params;
 struct FLobbyPrerequisitesEOS
 {
 	EOS_HLobby LobbyInterfaceHandle = {};
-	TWeakPtr<FAuthEOSGS> AuthInterface;
+	IEpicProductUserIdResolver* ResolverInterface;
 	FLobbyBucketIdEOS BucketId;
 	TSharedRef<FSchemaRegistry> SchemaRegistry;
 };
diff --git a/Engine/Plugins/Online/OnlineServicesEOSGS/Source/Private/Online/OnlineServicesEOSGS.cpp b/Engine/Plugins/Online/OnlineServicesEOSGS/Source/Private/Online/OnlineServicesEOSGS.cpp
--- a/Engine/Plugins/Online/OnlineServicesEOSGS/Source/Private/Online/OnlineServicesEOSGS.cpp
+++ b/Engine/Plugins/Online/OnlineServicesEOSGS/Source/Private/Online/OnlineServicesEOSGS.cpp
@@ -4,6 +4,7 @@
 
 #include "Online/AchievementsEOSGS.h"
 #include "Online/AuthEOSGS.h"
+#include "Online/EpicProductUserIdResolverEOSGS.h"
 #include "Online/LeaderboardsEOSGS.h"
 #include "Online/LobbiesEOSGS.h"
 #include "Online/OnlineIdEOSGS.h"
@@ -68,12 +69,13 @@ void FOnlineServicesEOSGS::RegisterComponents()
 {
 	Components.Register<FAchievementsEOSGS>(*this);
 	Components.Register<FAuthEOSGS>(*this);
+	Components.Register<FEpicProductUserIdResolverEOSGS>(*this);
 	Components.Register<FLeaderboardsEOSGS>(*this);
 	Components.Register<FLobbiesEOSGS>(*this);
-	Components.Register<FStatsEOSGS>(*this);
-	Components.Register<FSessionsEOSGS>(*this);
 	Components.Register<FPlayerReportsEOSGS>(*this);
 	Components.Register<FPlayerSanctionsEOSGS>(*this);
+	Components.Register<FStatsEOSGS>(*this);
+	Components.Register<FSessionsEOSGS>(*this);
 	if (EOS_Platform_GetTitleStorageInterface(*GetEOSPlatformHandle()))
 	{
 		Components.Register<FTitleFileEOSGS>(*this);
diff --git a/Engine/Plugins/Online/OnlineServicesEOSGS/Source/Private/Online/OnlineServicesEOSGSModule.cpp b/Engine/Plugins/Online/OnlineServicesEOSGS/Source/Private/Online/OnlineServicesEOSGSModule.cpp
--- a/Engine/Plugins/Online/OnlineServicesEOSGS/Source/Private/Online/OnlineServicesEOSGSModule.cpp
+++ b/Engine/Plugins/Online/OnlineServicesEOSGS/Source/Private/Online/OnlineServicesEOSGSModule.cpp
@@ -37,10 +37,7 @@ void FOnlineServicesEOSGSModule::StartupModule()
 	const bool bEnableEOSSplitPlugins = FParse::Param(FCommandLine::Get(), TEXT("EnableEOSSplitPlugins"));
 	if (bEnableEOSSplitPlugins)
 	{
-		FModuleManager::Get().LoadModuleChecked(TEXT("OnlineServicesEpicGameInternal"));
-		FModuleManager::Get().LoadModuleChecked(TEXT("OnlineServicesEpicAccountInternal"));
-
-		UE_LOG(LogOnlineServices, Verbose, TEXT("EnableEOSSplitPlugins launch parameter detected. Loading OnlineServicesEpicGameInternal and OnlineServicesEpicAccountInternal modules. Skipping FOnlineServicesEOSGSModule startup."));
+		UE_LOG(LogOnlineServices, Verbose, TEXT("EnableEOSSplitPlugins launch parameter detected. Skipping FOnlineServicesEOSGSModule startup."));
 		return;
 	}
 
diff --git a/Engine/Plugins/Online/OnlineServicesEOSGS/Source/Private/Online/PlayerReportsEOSGS.cpp b/Engine/Plugins/Online/OnlineServicesEOSGS/Source/Private/Online/PlayerReportsEOSGS.cpp
--- a/Engine/Plugins/Online/OnlineServicesEOSGS/Source/Private/Online/PlayerReportsEOSGS.cpp
+++ b/Engine/Plugins/Online/OnlineServicesEOSGS/Source/Private/Online/PlayerReportsEOSGS.cpp
@@ -77,7 +77,7 @@ TOnlineAsyncOpHandle<FSendPlayerReport> FPlayerReportsEOSGS::SendPlayerReport(FS
 		{
 			const FSendPlayerReport::Params& Params = InAsyncOp.GetParams();
 
-			if (!Services.Get<FAuthEOSGS>()->IsLoggedIn(Params.LocalAccountId))
+			if (!Services.Get<IAuth>()->IsLoggedIn(Params.LocalAccountId))
 			{
 				InAsyncOp.SetError(Errors::InvalidUser());
 			}
diff --git a/Engine/Plugins/Online/OnlineServicesEOSGS/Source/Private/Online/PlayerSanctionsEOSGS.cpp b/Engine/Plugins/Online/OnlineServicesEOSGS/Source/Private/Online/PlayerSanctionsEOSGS.cpp
--- a/Engine/Plugins/Online/OnlineServicesEOSGS/Source/Private/Online/PlayerSanctionsEOSGS.cpp
+++ b/Engine/Plugins/Online/OnlineServicesEOSGS/Source/Private/Online/PlayerSanctionsEOSGS.cpp
@@ -61,7 +61,7 @@ TOnlineAsyncOpHandle<FCreatePlayerSanctionAppeal> FPlayerSanctionsEOSGS::CreateP
 		{
 			const FCreatePlayerSanctionAppeal::Params& Params = InAsyncOp.GetParams();
 
-			if (!Services.Get<FAuthEOSGS>()->IsLoggedIn(Params.LocalAccountId))
+			if (!Services.Get<IAuth>()->IsLoggedIn(Params.LocalAccountId))
 			{
 					InAsyncOp.SetError(Errors::InvalidUser());
 			}
@@ -106,7 +106,7 @@ TOnlineAsyncOpHandle<FReadActivePlayerSanctions> FPlayerSanctionsEOSGS::ReadEntr
 		{
 			const FReadActivePlayerSanctions::Params& Params = InAsyncOp.GetParams();
 
-			if (!Services.Get<FAuthEOSGS>()->IsLoggedIn(Params.LocalAccountId))
+			if (!Services.Get<IAuth>()->IsLoggedIn(Params.LocalAccountId))
 			{
 				InAsyncOp.SetError(Errors::InvalidUser());
 			}
diff --git a/Engine/Plugins/Online/OnlineServicesEOSGS/Source/Private/Online/SessionsEOSGS.cpp b/Engine/Plugins/Online/OnlineServicesEOSGS/Source/Private/Online/SessionsEOSGS.cpp
--- a/Engine/Plugins/Online/OnlineServicesEOSGS/Source/Private/Online/SessionsEOSGS.cpp
+++ b/Engine/Plugins/Online/OnlineServicesEOSGS/Source/Private/Online/SessionsEOSGS.cpp
@@ -5,6 +5,7 @@
 #include "EOSShared.h"
 #include "IEOSSDKManager.h"
 #include "Online/AuthEOSGS.h"
+#include "Online/EpicProductUserIdResolver.h"
 #include "Online/LobbiesCommonTypes.h"
 #include "Online/NboSerializerEOSGSSvc.h"
 #include "Online/OnlineErrorEpicCommon.h"
@@ -1393,7 +1394,7 @@ TOnlineAsyncOpHandle<FBuildSessionFromDetailsHandle> FSessionsEOSGS::BuildSessio
 
 	if (GetIdsResult.IsOk())
 	{
-		Services.Get<FAuthEOSGS>()->ResolveAccountIds(Params.LocalAccountId, GetIdsResult.GetOkValue())
+		Services.Get<IEpicProductUserIdResolver>()->ResolveAccountIds(Params.LocalAccountId, GetIdsResult.GetOkValue())
 			.Next([this, WeakOp = Op->AsWeak(), Params = MoveTemp(Params)](TArray<FAccountId> ResolvedAccountIds) mutable
 		{
 			if (TOnlineAsyncOpPtr<FBuildSessionFromDetailsHandle> StrongOp = WeakOp.Pin())
@@ -1509,7 +1510,7 @@ TFuture<TOnlineResult<FAddSessionMember>> FSessionsEOSGS::AddSessionMemberImpl(c
 	TPromise<TOnlineResult<FAddSessionMember>> Promise;
 	TFuture<TOnlineResult<FAddSessionMember>> Future = Promise.GetFuture();
 
-	Services.Get<FAuthEOSGS>()->ResolveAccountIds(Params.LocalAccountId, { GetProductUserIdChecked(Params.LocalAccountId) })
+	Services.Get<IEpicProductUserIdResolver>()->ResolveAccountIds(Params.LocalAccountId, { GetProductUserIdChecked(Params.LocalAccountId) })
 		.Next([this, Promise = MoveTemp(Promise), Params](TArray<FAccountId> ResolvedAccountIds) mutable
 			{
 				if (!ResolvedAccountIds.IsEmpty())
@@ -1562,7 +1563,7 @@ TFuture<TOnlineResult<FRemoveSessionMember>> FSessionsEOSGS::RemoveSessionMember
 	TPromise<TOnlineResult<FRemoveSessionMember>> Promise;
 	TFuture<TOnlineResult<FRemoveSessionMember>> Future = Promise.GetFuture();
 
-	Services.Get<FAuthEOSGS>()->ResolveAccountIds(Params.LocalAccountId, { GetProductUserIdChecked(Params.LocalAccountId) })
+	Services.Get<IEpicProductUserIdResolver>()->ResolveAccountIds(Params.LocalAccountId, { GetProductUserIdChecked(Params.LocalAccountId) })
 		.Next([this, Promise = MoveTemp(Promise), Params](TArray<FAccountId> ResolvedAccountIds) mutable
 			{
 				if (!ResolvedAccountIds.IsEmpty())
diff --git a/Engine/Plugins/Online/OnlineServicesEOSGS/Source/Private/Online/StatsEOSGS.cpp b/Engine/Plugins/Online/OnlineServicesEOSGS/Source/Private/Online/StatsEOSGS.cpp
--- a/Engine/Plugins/Online/OnlineServicesEOSGS/Source/Private/Online/StatsEOSGS.cpp
+++ b/Engine/Plugins/Online/OnlineServicesEOSGS/Source/Private/Online/StatsEOSGS.cpp
@@ -38,7 +38,7 @@ TOnlineAsyncOpHandle<FUpdateStats> FStatsEOSGS::UpdateStats(FUpdateStats::Params
 {
 	TOnlineAsyncOpRef<FUpdateStats> Op = GetOp<FUpdateStats>(MoveTemp(Params));
 
-	if (!Services.Get<FAuthEOSGS>()->IsLoggedIn(Op->GetParams().LocalAccountId))
+	if (!Services.Get<IAuth>()->IsLoggedIn(Op->GetParams().LocalAccountId))
 	{
 		Op->SetError(Errors::InvalidUser());
 		return Op->GetHandle();
@@ -231,7 +231,7 @@ TOnlineAsyncOpHandle<FQueryStats> FStatsEOSGS::QueryStats(FQueryStats::Params&&
 		Op->Then([this](TOnlineAsyncOp<FQueryStats>& InAsyncOp)
 		{
 			const FQueryStats::Params& Params = InAsyncOp.GetParams();
-			if (!Services.Get<FAuthEOSGS>()->IsLoggedIn(Params.LocalAccountId))
+			if (!Services.Get<IAuth>()->IsLoggedIn(Params.LocalAccountId))
 			{
 				InAsyncOp.SetError(Errors::InvalidUser());
 			}
@@ -275,7 +275,7 @@ TOnlineAsyncOpHandle<FBatchQueryStats> FStatsEOSGS::BatchQueryStats(FBatchQueryS
 		Op->Then([this](TOnlineAsyncOp<FBatchQueryStats>& InAsyncOp)
 		{
 			const FBatchQueryStats::Params& Params = InAsyncOp.GetParams();
-			if (!Services.Get<FAuthEOSGS>()->IsLoggedIn(Params.LocalAccountId))
+			if (!Services.Get<IAuth>()->IsLoggedIn(Params.LocalAccountId))
 			{
 				InAsyncOp.SetError(Errors::InvalidUser());
 			}
diff --git a/Engine/Plugins/Online/OnlineServicesEOSGS/Source/Private/Online/TitleFileEOSGS.cpp b/Engine/Plugins/Online/OnlineServicesEOSGS/Source/Private/Online/TitleFileEOSGS.cpp
--- a/Engine/Plugins/Online/OnlineServicesEOSGS/Source/Private/Online/TitleFileEOSGS.cpp
+++ b/Engine/Plugins/Online/OnlineServicesEOSGS/Source/Private/Online/TitleFileEOSGS.cpp
@@ -42,7 +42,7 @@ TOnlineAsyncOpHandle<FTitleFileEnumerateFiles> FTitleFileEOSGS::EnumerateFiles(F
 		{
 			const FTitleFileEnumerateFiles::Params& Params = InAsyncOp.GetParams();
 
-			if (!Services.Get<FAuthEOSGS>()->IsLoggedIn(Params.LocalAccountId))
+			if (!Services.Get<IAuth>()->IsLoggedIn(Params.LocalAccountId))
 			{
 				InAsyncOp.SetError(Errors::InvalidUser());
 			}
@@ -122,7 +122,7 @@ TOnlineAsyncOpHandle<FTitleFileEnumerateFiles> FTitleFileEOSGS::EnumerateFiles(F
 
 TOnlineResult<FTitleFileGetEnumeratedFiles> FTitleFileEOSGS::GetEnumeratedFiles(FTitleFileGetEnumeratedFiles::Params&& Params)
 {
-	if (!Services.Get<FAuthEOSGS>()->IsLoggedIn(Params.LocalAccountId))
+	if (!Services.Get<IAuth>()->IsLoggedIn(Params.LocalAccountId))
 	{
 		return TOnlineResult<FTitleFileGetEnumeratedFiles>(Errors::InvalidUser());
 	}
@@ -167,7 +167,7 @@ TOnlineAsyncOpHandle<FTitleFileReadFile> FTitleFileEOSGS::ReadFile(FTitleFileRea
 		Op->Then([this](TOnlineAsyncOp<FTitleFileReadFile>& InAsyncOp)
 		{
 			const FTitleFileReadFile::Params& Params = InAsyncOp.GetParams();
-			if (!Services.Get<FAuthEOSGS>()->IsLoggedIn(Params.LocalAccountId))
+			if (!Services.Get<IAuth>()->IsLoggedIn(Params.LocalAccountId))
 			{
 				InAsyncOp.SetError(Errors::InvalidUser());
 			}
diff --git a/Engine/Plugins/Online/OnlineServicesEOSGS/Source/Private/Online/UserFileEOSGS.cpp b/Engine/Plugins/Online/OnlineServicesEOSGS/Source/Private/Online/UserFileEOSGS.cpp
--- a/Engine/Plugins/Online/OnlineServicesEOSGS/Source/Private/Online/UserFileEOSGS.cpp
+++ b/Engine/Plugins/Online/OnlineServicesEOSGS/Source/Private/Online/UserFileEOSGS.cpp
@@ -44,7 +44,7 @@ TOnlineAsyncOpHandle<FUserFileEnumerateFiles> FUserFileEOSGS::EnumerateFiles(FUs
 		{
 			const FUserFileEnumerateFiles::Params& Params = Op.GetParams();
 
-			if (!Services.Get<FAuthEOSGS>()->IsLoggedIn(Params.LocalAccountId))
+			if (!Services.Get<IAuth>()->IsLoggedIn(Params.LocalAccountId))
 			{
 				Op.SetError(Errors::InvalidUser());
 			}
@@ -125,7 +125,7 @@ TOnlineAsyncOpHandle<FUserFileEnumerateFiles> FUserFileEOSGS::EnumerateFiles(FUs
 
 TOnlineResult<FUserFileGetEnumeratedFiles> FUserFileEOSGS::GetEnumeratedFiles(FUserFileGetEnumeratedFiles::Params&& Params)
 {
-	if (!Services.Get<FAuthEOSGS>()->IsLoggedIn(Params.LocalAccountId))
+	if (!Services.Get<IAuth>()->IsLoggedIn(Params.LocalAccountId))
 	{
 		return TOnlineResult<FUserFileGetEnumeratedFiles>(Errors::InvalidUser());
 	}
@@ -170,7 +170,7 @@ TOnlineAsyncOpHandle<FUserFileReadFile> FUserFileEOSGS::ReadFile(FUserFileReadFi
 		Op->Then([this](TOnlineAsyncOp<FUserFileReadFile>& Op)
 		{
 			const FUserFileReadFile::Params& Params = Op.GetParams();
-			if (!Services.Get<FAuthEOSGS>()->IsLoggedIn(Params.LocalAccountId))
+			if (!Services.Get<IAuth>()->IsLoggedIn(Params.LocalAccountId))
 			{
 				Op.SetError(Errors::InvalidUser());
 			}
@@ -257,7 +257,7 @@ TOnlineAsyncOpHandle<FUserFileWriteFile> FUserFileEOSGS::WriteFile(FUserFileWrit
 	{
 		const FUserFileWriteFile::Params& Params = Op.GetParams();
 
-		if (!Services.Get<FAuthEOSGS>()->IsLoggedIn(Params.LocalAccountId))
+		if (!Services.Get<IAuth>()->IsLoggedIn(Params.LocalAccountId))
 		{
 			Op.SetError(Errors::InvalidUser());
 			Promise.EmplaceValue();
@@ -316,7 +316,7 @@ TOnlineAsyncOpHandle<FUserFileCopyFile> FUserFileEOSGS::CopyFile(FUserFileCopyFi
 	{
 		const FUserFileCopyFile::Params& Params = Op.GetParams();
 
-		if (!Services.Get<FAuthEOSGS>()->IsLoggedIn(Params.LocalAccountId))
+		if (!Services.Get<IAuth>()->IsLoggedIn(Params.LocalAccountId))
 		{
 			Op.SetError(Errors::InvalidUser());
 			Promise.EmplaceValue();
@@ -366,7 +366,7 @@ TOnlineAsyncOpHandle<FUserFileDeleteFile> FUserFileEOSGS::DeleteFile(FUserFileDe
 	{
 		const FUserFileDeleteFile::Params& Params = Op.GetParams();
 
-		if (!Services.Get<FAuthEOSGS>()->IsLoggedIn(Params.LocalAccountId))
+		if (!Services.Get<IAuth>()->IsLoggedIn(Params.LocalAccountId))
 		{
 			Op.SetError(Errors::InvalidUser());
 			Promise.EmplaceValue();
diff --git a/Engine/Plugins/Online/OnlineServicesEOSGS/Source/Public/Online/AuthEOSGS.h b/Engine/Plugins/Online/OnlineServicesEOSGS/Source/Public/Online/AuthEOSGS.h
--- a/Engine/Plugins/Online/OnlineServicesEOSGS/Source/Public/Online/AuthEOSGS.h
+++ b/Engine/Plugins/Online/OnlineServicesEOSGS/Source/Public/Online/AuthEOSGS.h
@@ -222,13 +222,6 @@ public:
 	ONLINESERVICESEOSGS_API virtual TOnlineResult<FAuthGetRelyingParty> GetRelyingParty(FAuthGetRelyingParty::Params&& Params) const override;
 	// End IAuth
 
-	// Begin FAuthEOSGS
-	ONLINESERVICESEOSGS_API virtual TFuture<FAccountId> ResolveAccountId(const FAccountId& LocalAccountId, const EOS_ProductUserId ProductUserId);
-	ONLINESERVICESEOSGS_API virtual TFuture<TArray<FAccountId>> ResolveAccountIds(const FAccountId& LocalAccountId, const TArray<EOS_ProductUserId>& ProductUserIds);
-	ONLINESERVICESEOSGS_API virtual TFunction<TFuture<FAccountId>(FOnlineAsyncOp& InAsyncOp, const EOS_ProductUserId& ProductUserId)> ResolveProductIdFn();
-	ONLINESERVICESEOSGS_API virtual TFunction<TFuture<TArray<FAccountId>>(FOnlineAsyncOp& InAsyncOp, const TArray<EOS_ProductUserId>& ProductUserIds)> ResolveProductIdsFn();
-	// End FAuthEOSGS
-
 protected:
 	// internal operations.
 
diff --git a/Engine/Plugins/Online/OnlineServicesEOSGS/Source/Public/Online/EpicProductUserIdResolverEOSGS.h b/Engine/Plugins/Online/OnlineServicesEOSGS/Source/Public/Online/EpicProductUserIdResolverEOSGS.h
new file mode 100644
--- /dev/null
+++ b/Engine/Plugins/Online/OnlineServicesEOSGS/Source/Public/Online/EpicProductUserIdResolverEOSGS.h
@@ -0,0 +1,21 @@
+// Copyright Epic Games, Inc. All Rights Reserved.
+
+#pragma once
+
+#include "Online/EpicProductUserIdResolver.h"
+#include "Online/OnlineComponent.h"
+
+namespace UE::Online {
+
+class FEpicProductUserIdResolverEOSGS : public TOnlineComponent<IEpicProductUserIdResolver>
+{
+public:
+	FEpicProductUserIdResolverEOSGS(class FOnlineServicesEOSGS& InServices);
+
+	// Begin IEpicProductUserIdResolver
+	virtual TFuture<FAccountId> ResolveAccountId(const FAccountId LocalAccountId, const EOS_ProductUserId ProductUserId) override;
+	virtual TFuture<TArray<FAccountId>> ResolveAccountIds(const FAccountId LocalAccountId, const TArray<EOS_ProductUserId>& ProductUserIds) override;
+	// End IEpicProductUserIdResolver
+};
+
+/* UE::Online */ }
diff --git a/Engine/Plugins/Online/OnlineServicesEpicCommon/Source/OnlineServicesEpicCommon/Private/Online/EpicAccountIdResolver.cpp b/Engine/Plugins/Online/OnlineServicesEpicCommon/Source/OnlineServicesEpicCommon/Private/Online/EpicAccountIdResolver.cpp
new file mode 100644
--- /dev/null
+++ b/Engine/Plugins/Online/OnlineServicesEpicCommon/Source/OnlineServicesEpicCommon/Private/Online/EpicAccountIdResolver.cpp
@@ -0,0 +1,36 @@
+// Copyright Epic Games, Inc. All Rights Reserved.
+
+#include "Online/EpicAccountIdResolver.h"
+
+#include "Online/OnlineAsyncOp.h"
+
+namespace UE::Online
+{
+
+TFunction<TFuture<FAccountId>(FOnlineAsyncOp& InAsyncOp, const EOS_EpicAccountId EpicAccountId)> IEpicAccountIdResolver::ResolveEpicAccountIdFn()
+{
+	return [this](FOnlineAsyncOp& InAsyncOp, const EOS_EpicAccountId& EpicAccountId)
+	{
+		const FAccountId* LocalAccountIdPtr = InAsyncOp.Data.Get<FAccountId>(TEXT("LocalAccountId"));
+		if (!ensure(LocalAccountIdPtr))
+		{
+			return MakeFulfilledPromise<FAccountId>().GetFuture();
+		}
+		return ResolveAccountId(*LocalAccountIdPtr, EpicAccountId);
+	};
+}
+
+TFunction<TFuture<TArray<FAccountId>>(FOnlineAsyncOp& InAsyncOp, const TArray<EOS_EpicAccountId>& EpicAccountIds)> IEpicAccountIdResolver::ResolveEpicAccountIdsFn()
+{
+	return [this](FOnlineAsyncOp& InAsyncOp, const TArray<EOS_EpicAccountId>& EpicAccountIds)
+	{
+		const FAccountId* LocalAccountIdPtr = InAsyncOp.Data.Get<FAccountId>(TEXT("LocalAccountId"));
+		if (!ensure(LocalAccountIdPtr))
+		{
+			return MakeFulfilledPromise<TArray<FAccountId>>().GetFuture();
+		}
+		return ResolveAccountIds(*LocalAccountIdPtr, EpicAccountIds);
+	};
+}
+
+} // namespace UE::Online
\ No newline at end of file
diff --git a/Engine/Plugins/Online/OnlineServicesEpicCommon/Source/OnlineServicesEpicCommon/Private/Online/EpicProductUserIdResolver.cpp b/Engine/Plugins/Online/OnlineServicesEpicCommon/Source/OnlineServicesEpicCommon/Private/Online/EpicProductUserIdResolver.cpp
new file mode 100644
--- /dev/null
+++ b/Engine/Plugins/Online/OnlineServicesEpicCommon/Source/OnlineServicesEpicCommon/Private/Online/EpicProductUserIdResolver.cpp
@@ -0,0 +1,36 @@
+// Copyright Epic Games, Inc. All Rights Reserved.
+
+#include "Online/EpicProductUserIdResolver.h"
+
+#include "Online/OnlineAsyncOp.h"
+
+namespace UE::Online
+{
+
+TFunction<TFuture<FAccountId>(FOnlineAsyncOp& InAsyncOp, const EOS_ProductUserId ProductUserId)> IEpicProductUserIdResolver::ResolveProductUserIdFn()
+{
+	return [this](FOnlineAsyncOp& InAsyncOp, const EOS_ProductUserId& ProductUserId)
+	{
+		const FAccountId* LocalAccountIdPtr = InAsyncOp.Data.Get<FAccountId>(TEXT("LocalAccountId"));
+		if (!ensure(LocalAccountIdPtr))
+		{
+			return MakeFulfilledPromise<FAccountId>().GetFuture();
+		}
+		return ResolveAccountId(*LocalAccountIdPtr, ProductUserId);
+	};
+}
+
+TFunction<TFuture<TArray<FAccountId>>(FOnlineAsyncOp& InAsyncOp, const TArray<EOS_ProductUserId>& ProductUserIds)> IEpicProductUserIdResolver::ResolveProductUserIdsFn()
+{
+	return [this](FOnlineAsyncOp& InAsyncOp, const TArray<EOS_ProductUserId>& ProductUserIds)
+	{
+		const FAccountId* LocalAccountIdPtr = InAsyncOp.Data.Get<FAccountId>(TEXT("LocalAccountId"));
+		if (!ensure(LocalAccountIdPtr))
+		{
+			return MakeFulfilledPromise<TArray<FAccountId>>().GetFuture();
+		}
+		return ResolveAccountIds(*LocalAccountIdPtr, ProductUserIds);
+	};
+}
+
+} // namespace UE::Online
\ No newline at end of file
diff --git a/Engine/Plugins/Online/OnlineServicesEpicCommon/Source/OnlineServicesEpicCommon/Public/Online/EpicAccountIdResolver.h b/Engine/Plugins/Online/OnlineServicesEpicCommon/Source/OnlineServicesEpicCommon/Public/Online/EpicAccountIdResolver.h
new file mode 100644
--- /dev/null
+++ b/Engine/Plugins/Online/OnlineServicesEpicCommon/Source/OnlineServicesEpicCommon/Public/Online/EpicAccountIdResolver.h
@@ -0,0 +1,24 @@
+// Copyright Epic Games, Inc. All Rights Reserved.
+
+#pragma once
+
+#include "Async/Future.h"
+#include "EOSShared.h"
+#include "Online/Auth.h"
+#include "Online/CoreOnline.h"
+
+namespace UE::Online
+{
+
+/** Exists to provide components that are instantiated in both OnlineServicesEOS and OnlineServicesEpicAccount a resolve path that works in both. */
+class IEpicAccountIdResolver
+{
+public:
+	virtual TFuture<FAccountId> ResolveAccountId(const FAccountId LocalAccountId, const EOS_EpicAccountId EpicAccountId) = 0;
+	virtual TFuture<TArray<FAccountId>> ResolveAccountIds(const FAccountId LocalAccountId, const TArray<EOS_EpicAccountId>& EpicAccountIds) = 0;
+
+	ONLINESERVICESEPICCOMMON_API TFunction<TFuture<FAccountId>(class FOnlineAsyncOp& InAsyncOp, const EOS_EpicAccountId EpicAccountId)> ResolveEpicAccountIdFn();
+	ONLINESERVICESEPICCOMMON_API TFunction<TFuture<TArray<FAccountId>>(class FOnlineAsyncOp& InAsyncOp, const TArray<EOS_EpicAccountId>& EpicAccountIds)> ResolveEpicAccountIdsFn();
+};
+
+} // namespace UE::Online
\ No newline at end of file
diff --git a/Engine/Plugins/Online/OnlineServicesEpicCommon/Source/OnlineServicesEpicCommon/Public/Online/EpicProductUserIdResolver.h b/Engine/Plugins/Online/OnlineServicesEpicCommon/Source/OnlineServicesEpicCommon/Public/Online/EpicProductUserIdResolver.h
new file mode 100644
--- /dev/null
+++ b/Engine/Plugins/Online/OnlineServicesEpicCommon/Source/OnlineServicesEpicCommon/Public/Online/EpicProductUserIdResolver.h
@@ -0,0 +1,24 @@
+// Copyright Epic Games, Inc. All Rights Reserved.
+
+#pragma once
+
+#include "Async/Future.h"
+#include "EOSShared.h"
+#include "Online/Auth.h"
+#include "Online/CoreOnline.h"
+
+namespace UE::Online
+{
+
+/** Exists to provide components that are instantiated in both OnlineServicesEOSGS and OnlineServicesEpicGame a resolve path that works in both. */
+class IEpicProductUserIdResolver
+{
+public:
+	virtual TFuture<FAccountId> ResolveAccountId(const FAccountId LocalAccountId, const EOS_ProductUserId ProductUserId) = 0;
+	virtual TFuture<TArray<FAccountId>> ResolveAccountIds(const FAccountId LocalAccountId, const TArray<EOS_ProductUserId>& ProductUserIds) = 0;
+
+	ONLINESERVICESEPICCOMMON_API TFunction<TFuture<FAccountId>(class FOnlineAsyncOp& InAsyncOp, const EOS_ProductUserId ProductUserId)> ResolveProductUserIdFn();
+	ONLINESERVICESEPICCOMMON_API TFunction<TFuture<TArray<FAccountId>>(class FOnlineAsyncOp& InAsyncOp, const TArray<EOS_ProductUserId>& ProductUserIds)> ResolveProductUserIdsFn();
+};
+
+} // namespace UE::Online
\ No newline at end of file
