- Fix -enableeossplitplugins causing failures in OnlineServicesEOS etc on shutdown due to it not establishing a loads-after module dependency on OnlineServicesInterface in that case.

#rnx

diff --git a/Engine/Plugins/Online/OnlineServicesEOS/Source/Private/Online/OnlineServicesEOSModule.cpp b/Engine/Plugins/Online/OnlineServicesEOS/Source/Private/Online/OnlineServicesEOSModule.cpp
--- a/Engine/Plugins/Online/OnlineServicesEOS/Source/Private/Online/OnlineServicesEOSModule.cpp
+++ b/Engine/Plugins/Online/OnlineServicesEOS/Source/Private/Online/OnlineServicesEOSModule.cpp
@@ -35,6 +35,8 @@ int FOnlineServicesEOSModule::GetRegistryPriority()
 
 void FOnlineServicesEOSModule::StartupModule()
 {
+	FModuleManager::Get().LoadModuleChecked(TEXT("OnlineServicesInterface"));
+
 	const bool bEnableEOSSplitPlugins = FParse::Param(FCommandLine::Get(), TEXT("EnableEOSSplitPlugins"));
 	if (bEnableEOSSplitPlugins)
 	{
@@ -44,13 +46,6 @@ void FOnlineServicesEOSModule::StartupModule()
 
 	FModuleManager::Get().LoadModuleChecked(TEXT("OnlineServicesEOSGS"));
 
-	// Making sure we load the module at this point will avoid errors while cooking
-	const FName OnlineServicesInterfaceModuleName = TEXT("OnlineServicesInterface");
-	if (!FModuleManager::Get().IsModuleLoaded(OnlineServicesInterfaceModuleName))
-	{
-		FModuleManager::Get().LoadModuleChecked(OnlineServicesInterfaceModuleName);
-	}
-
 	FOnlineServicesRegistry::Get().RegisterServicesFactory(EOnlineServices::Epic, MakeUnique<FOnlineServicesFactoryEOS>(), GetRegistryPriority());
 	FOnlineIdRegistryRegistry::Get().RegisterAccountIdRegistry(EOnlineServices::Epic, &FOnlineAccountIdRegistryEOS::Get(), GetRegistryPriority());
 }
