Added support to TNestedVariant in related usage by Presence prototype changes
Added FOnlinePartyJoinInfoMcpV2JsonSerializerV2 to produce a shorter Json string
Added logic to FPresenceEOS to serialize and deserialize TNestedVariant attributes into string

diff --git a/Engine/Plugins/Online/EOSShared/Source/EOSShared/Private/EOSShared.cpp b/Engine/Plugins/Online/EOSShared/Source/EOSShared/Private/EOSShared.cpp
--- a/Engine/Plugins/Online/EOSShared/Source/EOSShared/Private/EOSShared.cpp
+++ b/Engine/Plugins/Online/EOSShared/Source/EOSShared/Private/EOSShared.cpp
@@ -678,15 +678,15 @@ EOS_OnlinePlatformType EOSOnlinePlatformTypeFromString(const FStringView& InStri
 	}
 	else if (InString == TEXT("PSN"))
 	{	
-		return 1000; //EOS_OPT_PSN;
+		return EOS_OPT_PSN;
 	}
 	else if (InString == TEXT("Switch"))
 	{	
-		return 2000; //EOS_OPT_SWITCH;
+		return EOS_OPT_Nintendo;
 	}
 	else if (InString == TEXT("XBL"))
 	{	
-		return 3000; //EOS_OPT_XBL;
+		return EOS_OPT_XBL;
 	}
 	else
 	{
diff --git a/Engine/Plugins/Online/OnlineFramework/Source/Party/Private/User/SocialUser.cpp b/Engine/Plugins/Online/OnlineFramework/Source/Party/Private/User/SocialUser.cpp
--- a/Engine/Plugins/Online/OnlineFramework/Source/Party/Private/User/SocialUser.cpp
+++ b/Engine/Plugins/Online/OnlineFramework/Source/Party/Private/User/SocialUser.cpp
@@ -139,53 +139,56 @@ public:
 					bHasExecuted = true;
 
 					const FUniqueNetIdRepl& LocalUserNetId = Toolkit->GetLocalUserNetId(SubsystemType);
-					if (const FCommonAccountPtr LocalUserCommonAccountPtr = GetCommonAccountFromV1(CommonConfig, LocalUserNetId->AsShared(), GetFirstFrameworkInstanceName(LocalUserNetId)))
+					if(LocalUserNetId.IsValid())
 					{
-						FAccountId LocalUserAccountId = LocalUserCommonAccountPtr->GetId(StrongToolkit->GetPresenceFrameworkInstance());
+						if (const FCommonAccountPtr LocalUserCommonAccountPtr = GetCommonAccountFromV1(CommonConfig, LocalUserNetId->AsShared(), GetFirstFrameworkInstanceName(LocalUserNetId)))
+						{
+							FAccountId LocalUserAccountId = LocalUserCommonAccountPtr->GetId(StrongToolkit->GetPresenceFrameworkInstance());
 
-						TArray<FAccountId> SocialUserAccountIds;
-						CompletionCallbacksByUserId.GenerateKeyArray(SocialUserAccountIds);
-						UE_LOG(LogParty, Log, TEXT("FSocialQuery_UserInfoV2 executing for [%d] users on subsystem [%s]"), SocialUserAccountIds.Num(), ToString(SubsystemType));
+							TArray<FAccountId> SocialUserAccountIds;
+							CompletionCallbacksByUserId.GenerateKeyArray(SocialUserAccountIds);
+							UE_LOG(LogParty, Log, TEXT("FSocialQuery_UserInfoV2 executing for [%d] users on subsystem [%s]"), SocialUserAccountIds.Num(), ToString(SubsystemType));
 
-						UserInfo->QueryUserInfo({ LocalUserAccountId, SocialUserAccountIds })
-							.OnComplete([this, LocalUserAccountId, SocialUserAccountIds, UserInfoWeakPtr = UserInfo.ToWeakPtr()](const TOnlineResult<FQueryUserInfo>& Result)
-							{
-								if (Result.IsOk())
+							UserInfo->QueryUserInfo({ LocalUserAccountId, SocialUserAccountIds })
+								.OnComplete([this, LocalUserAccountId, SocialUserAccountIds, UserInfoWeakPtr = UserInfo.ToWeakPtr()](const TOnlineResult<FQueryUserInfo>& Result)
 								{
-									UE_LOG(LogParty, Log, TEXT("FSocialQuery_UserInfoV2 completed query on services [%s]"), ToString(SubsystemType));
+									if (Result.IsOk())
+									{
+										UE_LOG(LogParty, Log, TEXT("FSocialQuery_UserInfoV2 completed query on services [%s]"), ToString(SubsystemType));
 											
-									using namespace UE::Online;
+										using namespace UE::Online;
 
-									if (IUserInfoPtr UserInfo = UserInfoWeakPtr.Pin())
-									{
-										// Notify users of the query completion
-										for (const FAccountId& UserAccountId : SocialUserAccountIds)
+										if (IUserInfoPtr UserInfo = UserInfoWeakPtr.Pin())
 										{
-											TOnlineResult<FGetUserInfo> GetUserInfoResult = UserInfo->GetUserInfo({ LocalUserAccountId, UserAccountId });
-											if (GetUserInfoResult.IsOk())
+											// Notify users of the query completion
+											for (const FAccountId& UserAccountId : SocialUserAccountIds)
 											{
-												if (const FOnQueryComplete* CompletionCallback = CompletionCallbacksByUserId.Find(UserAccountId))
+												TOnlineResult<FGetUserInfo> GetUserInfoResult = UserInfo->GetUserInfo({ LocalUserAccountId, UserAccountId });
+												if (GetUserInfoResult.IsOk())
 												{
-													CompletionCallback->ExecuteIfBound(SubsystemType, true, MakeValue(GetUserInfoResult.GetOkValue().UserInfo));
+													if (const FOnQueryComplete* CompletionCallback = CompletionCallbacksByUserId.Find(UserAccountId))
+													{
+														CompletionCallback->ExecuteIfBound(SubsystemType, true, MakeValue(GetUserInfoResult.GetOkValue().UserInfo));
+													}
 												}
-											}
-											else
-											{
-												if (const FOnQueryComplete* CompletionCallback = CompletionCallbacksByUserId.Find(UserAccountId))
+												else
 												{
-													CompletionCallback->ExecuteIfBound(SubsystemType, false, MakeError(GetUserInfoResult.GetErrorValue()));
+													if (const FOnQueryComplete* CompletionCallback = CompletionCallbacksByUserId.Find(UserAccountId))
+													{
+														CompletionCallback->ExecuteIfBound(SubsystemType, false, MakeError(GetUserInfoResult.GetErrorValue()));
+													}
 												}
 											}
 										}
-									}
 
-									OnQueryCompleted.ExecuteIfBound(GetQueryId(), AsShared());
-								}
-								else
-								{
-									UE_LOG(LogParty, Log, TEXT("FSocialQuery_UserInfoV2 failed query on services [%s] with error [%s]"), ToString(SubsystemType), *Result.GetErrorValue().GetLogString());
-								}
-							});
+										OnQueryCompleted.ExecuteIfBound(GetQueryId(), AsShared());
+									}
+									else
+									{
+										UE_LOG(LogParty, Log, TEXT("FSocialQuery_UserInfoV2 failed query on services [%s] with error [%s]"), ToString(SubsystemType), *Result.GetErrorValue().GetLogString());
+									}
+								});
+						}
 					}
 				}
 			}
@@ -655,24 +658,19 @@ FString USocialUser::GetDisplayName() const
 
 	if(const FSubsystemUserInfo* SubsystemInfo = SubsystemInfoByType.Find(ESocialSubsystem::Primary))
 	{
-		// Temporary redundancy check
-		FString DisplayNameV1, DisplayNameV2;
-
 		// We always retrieve the V1 name during the transition period, for the redundancy check
 		if (TSharedPtr<FOnlineUser> PrimaryUserInfo = SubsystemInfo->UserInfo.Pin())
 		{
 			// The primary user info has knowledge of display names on all linked accounts
 			const FUserPlatform& UserCurrentPlatform = GetCurrentPlatform();
-			DisplayNameV1 = PrimaryUserInfo->GetDisplayName(UserCurrentPlatform.ToString());
+			DisplayName = PrimaryUserInfo->GetDisplayName(UserCurrentPlatform.ToString());
 		}
 		else
 		{
 			// We don't have primary user info (so we're not even initialized yet!), but a good-faith effort is to see if we have a platform name
-			DisplayNameV1 = GetDisplayName(ESocialSubsystem::Platform);
+			DisplayName = GetDisplayName(ESocialSubsystem::Platform);
 		}
 
-		DisplayName = DisplayNameV1;
-
 		if (IsOnlineServicesPresenceEnabled())
 		{
 			using namespace UE::Online;
@@ -695,6 +693,8 @@ FString USocialUser::GetDisplayName() const
 					{
 						if (IUserInfoPtr UserInfo = OnlineServices->GetUserInfoInterface())
 						{
+							FString DisplayNameV2;
+
 							TOnlineResult<FGetUserInfo> GetUserInfoResult = UserInfo->GetUserInfo({ LocalUserAccountId, SocialUserAccountId }); //, GetCurrentPlatform().ToString()}); // TODO: We probably need a translator from GetCurrentPlatform to what EOS understands
 							if (GetUserInfoResult.IsOk())
 							{
@@ -707,13 +707,16 @@ FString USocialUser::GetDisplayName() const
 							}
 
 #if !UE_BUILD_SHIPPING
-							if (DisplayNameV1 != DisplayNameV2)
+							// Temporary redundancy check
+							if (DisplayName != DisplayNameV2)
 							{
-								UE_LOG(LogParty, Warning, TEXT("DisplayName disparity: DisplayNameV1: %s. DisplayNameV2: %s"), *DisplayNameV1, *DisplayNameV2);
+								UE_LOG(LogParty, Warning, TEXT("DisplayName disparity: DisplayNameV1: %s. DisplayNameV2: %s"), *DisplayName, *DisplayNameV2);
 							}
 #endif
-
-							DisplayName = DisplayNameV2;
+							if (!DisplayNameV2.IsEmpty())
+							{
+								DisplayName = DisplayNameV2;
+							}
 						}
 					}
 				}
@@ -726,10 +729,18 @@ FString USocialUser::GetDisplayName() const
 
 FString USocialUser::GetDisplayName(ESocialSubsystem SubsystemType) const
 {
+	FString DisplayName;
+
 	if (const FSubsystemUserInfo* SubsystemInfo = SubsystemInfoByType.Find(SubsystemType))
 	{
 		if (TSharedPtr<FOnlineUser> PrimaryUserInfo = SubsystemInfo->UserInfo.Pin())
 		{
+			// Temporary redundancy check
+			FString DisplayNameV1, DisplayNameV2;
+
+			DisplayNameV1 = PrimaryUserInfo->GetDisplayName();
+			DisplayName = DisplayNameV1;
+
 			if (IsOnlineServicesPresenceEnabled() && SubsystemType == ESocialSubsystem::Primary) // We don't want V2 functionality in Platform for now
 			{
 				using namespace UE::Online;
@@ -755,20 +766,28 @@ FString USocialUser::GetDisplayName(ESocialSubsystem SubsystemType) const
 								TOnlineResult<FGetUserInfo> GetUserInfoResult = UserInfo->GetUserInfo({ LocalUserAccountId, SocialUserAccountId });
 								if (GetUserInfoResult.IsOk())
 								{
-									return GetUserInfoResult.GetOkValue().UserInfo->DisplayName;
+									DisplayNameV2 = GetUserInfoResult.GetOkValue().UserInfo->DisplayName;
+#if !UE_BUILD_SHIPPING
+									// Temporary redundancy check
+									if (DisplayNameV1 != DisplayNameV2)
+									{
+										UE_LOG(LogParty, Warning, TEXT("DisplayName disparity: DisplayNameV1: %s. DisplayNameV2: %s"), *DisplayNameV1, *DisplayNameV2);
+									}
+#endif
+									if (!DisplayNameV2.IsEmpty())
+									{
+										DisplayName = DisplayNameV2;
+									}
 								}
 							}
 						}
 					}
 				}
 			}
-			else
-			{
-				return PrimaryUserInfo->GetDisplayName();
-			}
 		}
 	}
-	return FString();
+
+	return DisplayName;
 }
 
 FString USocialUser::GetNickname() const
diff --git a/Engine/Plugins/Online/OnlineServices/Source/OnlineServicesCommon/Private/Online/PresenceCommon.cpp b/Engine/Plugins/Online/OnlineServices/Source/OnlineServicesCommon/Private/Online/PresenceCommon.cpp
--- a/Engine/Plugins/Online/OnlineServices/Source/OnlineServicesCommon/Private/Online/PresenceCommon.cpp
+++ b/Engine/Plugins/Online/OnlineServices/Source/OnlineServicesCommon/Private/Online/PresenceCommon.cpp
@@ -157,7 +157,7 @@ TOnlineAsyncOpHandle<FPartialUpdatePresence> FPresenceCommon::PartialUpdatePrese
 			NewPresence->Joinability = Params.Mutations.Joinability.GetValue();
 		}
 
-		for (const TPair<FString, FPresenceVariant>& UpdatedProperty : Params.Mutations.UpdatedProperties)
+		for (const TPair<FString, FPresenceProperty>& UpdatedProperty : Params.Mutations.UpdatedProperties)
 		{
 			NewPresence->Properties.Add(UpdatedProperty.Key, UpdatedProperty.Value);
 		}
diff --git a/Engine/Plugins/Online/OnlineServices/Source/OnlineServicesInterface/Private/Online/OnlineServices.cpp b/Engine/Plugins/Online/OnlineServices/Source/OnlineServicesInterface/Private/Online/OnlineServices.cpp
--- a/Engine/Plugins/Online/OnlineServices/Source/OnlineServicesInterface/Private/Online/OnlineServices.cpp
+++ b/Engine/Plugins/Online/OnlineServices/Source/OnlineServicesInterface/Private/Online/OnlineServices.cpp
@@ -13,6 +13,53 @@ DEFINE_LOG_CATEGORY(LogOnlineServices);
 
 namespace UE::Online {
 
+const TCHAR* LexToString(EOnlinePlatformType EnumVal)
+{
+	switch (EnumVal)
+	{
+	case EOnlinePlatformType::Epic:			return TEXT("Epic");
+	case EOnlinePlatformType::Steam:			return TEXT("Steam");
+	case EOnlinePlatformType::PSN:			return TEXT("PSN");
+	case EOnlinePlatformType::Nintendo:		return TEXT("Nintendo");
+	case EOnlinePlatformType::XBL:			return TEXT("XBL");
+	default:										checkNoEntry(); // Intentional fall-through
+	case EOnlinePlatformType::Unknown:		return TEXT("Unknown");
+	}
+}
+
+void LexFromString(EOnlinePlatformType& OutPlatformType, const TCHAR* InStr)
+{
+	if (FCString::Stricmp(InStr, TEXT("Epic")) == 0)
+	{
+		OutPlatformType = EOnlinePlatformType::Epic;
+	}
+	else if (FCString::Stricmp(InStr, TEXT("Steam")) == 0)
+	{
+		OutPlatformType = EOnlinePlatformType::Steam;
+	}
+	else if (FCString::Stricmp(InStr, TEXT("PSN")) == 0)
+	{
+		OutPlatformType = EOnlinePlatformType::PSN;
+	}
+	else if (FCString::Stricmp(InStr, TEXT("Nintendo")) == 0)
+	{
+		OutPlatformType = EOnlinePlatformType::Nintendo;
+	}
+	else if (FCString::Stricmp(InStr, TEXT("XBL")) == 0)
+	{
+		OutPlatformType = EOnlinePlatformType::XBL;
+	}
+	else if (FCString::Stricmp(InStr, TEXT("Unknown")) == 0)
+	{
+		OutPlatformType = EOnlinePlatformType::Unknown;
+	}
+	else
+	{
+		checkNoEntry();
+		OutPlatformType = EOnlinePlatformType::Unknown;
+	}
+}
+
 int32 GetBuildUniqueId()
 {
 	static bool bStaticCheck = false;
diff --git a/Engine/Plugins/Online/OnlineServices/Source/OnlineServicesInterface/Public/Online/OnlineServices.h b/Engine/Plugins/Online/OnlineServices/Source/OnlineServicesInterface/Public/Online/OnlineServices.h
--- a/Engine/Plugins/Online/OnlineServices/Source/OnlineServicesInterface/Public/Online/OnlineServices.h
+++ b/Engine/Plugins/Online/OnlineServices/Source/OnlineServicesInterface/Public/Online/OnlineServices.h
@@ -13,6 +13,18 @@ class FString;
 
 namespace UE::Online {
 
+enum class EOnlinePlatformType : uint8
+{
+	Epic,
+	Steam,
+	PSN,
+	Nintendo,
+	XBL,
+	Unknown
+};
+ONLINESERVICESINTERFACE_API const TCHAR* LexToString(EOnlinePlatformType EnumVal);
+ONLINESERVICESINTERFACE_API void LexFromString(EOnlinePlatformType& OutPlatformType, const TCHAR* InStr);
+
 // Interfaces
 using IOnlineServicesPtr = TSharedPtr<class IOnlineServices>;
 using IAchievementsPtr = TSharedPtr<class IAchievements>;
diff --git a/Engine/Plugins/Online/OnlineServices/Source/OnlineServicesInterface/Public/Online/Presence.h b/Engine/Plugins/Online/OnlineServices/Source/OnlineServicesInterface/Public/Online/Presence.h
--- a/Engine/Plugins/Online/OnlineServices/Source/OnlineServicesInterface/Public/Online/Presence.h
+++ b/Engine/Plugins/Online/OnlineServices/Source/OnlineServicesInterface/Public/Online/Presence.h
@@ -3,9 +3,11 @@
 #pragma once
 
 #include "Online/CoreOnline.h"
+#include "Online/OnlineError.h"
 #include "Online/OnlineErrorCode.h"
 #include "Online/OnlineMeta.h"
-#include "Online/OnlineError.h"
+#include "Online/OnlineServices.h"
+#include "Online/NestedVariantJson.h"
 
 namespace UE::Online { template <typename DelegateSignature> class TOnlineEvent; }
 namespace UE::Online { template <typename OpType> class TOnlineAsyncOpHandle; }
@@ -14,9 +16,11 @@ namespace UE::Online::Meta { template <typename StructType> struct TStructDetail
 
 #define LOCTEXT_NAMESPACE "Presence"
 
+#define SENT_TIME_ATTRIBUTE_KEY FString(TEXT("SentTimeAttributeKey"))
+
 namespace UE::Online {
 
-using FPresenceVariant = FString; // TODO:  Import FVariantData
+using FPresenceVariant = FNestedVariantJson;
 
 namespace Errors {
 UE_ONLINE_ERROR_CATEGORY(Presence, Engine, 0x3, "Presence")
@@ -79,7 +83,10 @@ enum class EUserPresenceGameStatus : uint8
 ONLINESERVICESINTERFACE_API const TCHAR* LexToString(EUserPresenceGameStatus EnumVal);
 ONLINESERVICESINTERFACE_API void LexFromString(EUserPresenceGameStatus& OutGameStatus, const TCHAR* InStr);
 
-typedef TMap<FString, FPresenceVariant> FPresenceProperties;
+using FPresenceProperties = FPresenceVariant::FMap;
+using FPresencePropertiesRef = FPresenceVariant::FMapRef;
+using FPresencePropertiesPtr = FPresenceVariant::FMapPtr;
+using FPresenceProperty = FPresenceProperties::ValueType;
 
 struct FUserPresence
 {
@@ -95,6 +102,8 @@ struct FUserPresence
 	FString StatusString;
 	/** Game-defined representation of the current game state (e.g. "Squads- 4 Remaining" or "Level 5 Warrior") */
 	FString RichPresenceString;
+	/** Type of the user's current online platform */
+	EOnlinePlatformType PlatformType = EOnlinePlatformType::Unknown;
 	/** Session keys, and any platform specific presence properties */
 	FPresenceProperties Properties;
 };
@@ -239,7 +248,7 @@ inline FPartialUpdatePresence::Params::FMutations& FPartialUpdatePresence::Param
 	}
 
 	// Insert any updated keys / remove from removed
-	for (TPair<FString, FPresenceVariant>& NewUpdatedProperty : NewParams.UpdatedProperties)
+	for (TPair<FString, FPresenceProperty>& NewUpdatedProperty : NewParams.UpdatedProperties)
 	{
 		RemovedProperties.Remove(NewUpdatedProperty.Key);
 		UpdatedProperties.Emplace(MoveTemp(NewUpdatedProperty.Key), MoveTemp(NewUpdatedProperty.Value));
@@ -286,7 +295,7 @@ inline TSharedRef<FUserPresence> ApplyPresenceMutations(const FUserPresence& Bas
 	}
 
 	// Insert any updated keys / remove from removed
-	for (const TPair<FString, FPresenceVariant>& NewUpdatedProperty : Mutations.UpdatedProperties)
+	for (const TPair<FString, FPresenceProperty>& NewUpdatedProperty : Mutations.UpdatedProperties)
 	{
 		NewPresence->Properties.Add(NewUpdatedProperty.Key, NewUpdatedProperty.Value);
 	}
@@ -371,6 +380,7 @@ BEGIN_ONLINE_STRUCT_META(FUserPresence)
 	ONLINE_STRUCT_FIELD(FUserPresence, GameStatus),
 	ONLINE_STRUCT_FIELD(FUserPresence, StatusString),
 	ONLINE_STRUCT_FIELD(FUserPresence, RichPresenceString),
+	ONLINE_STRUCT_FIELD(FUserPresence, PlatformType),
 	ONLINE_STRUCT_FIELD(FUserPresence, Properties)
 END_ONLINE_STRUCT_META()
 
diff --git a/Engine/Plugins/Online/OnlineServices/Source/OnlineServicesInterface/Public/Online/UserInfo.h b/Engine/Plugins/Online/OnlineServices/Source/OnlineServicesInterface/Public/Online/UserInfo.h
--- a/Engine/Plugins/Online/OnlineServices/Source/OnlineServicesInterface/Public/Online/UserInfo.h
+++ b/Engine/Plugins/Online/OnlineServices/Source/OnlineServicesInterface/Public/Online/UserInfo.h
@@ -4,6 +4,7 @@
 
 #include "Online/CoreOnline.h"
 #include "Online/OnlineMeta.h"
+#include "Online/OnlineServices.h"
 
 namespace UE::Online { template <typename OpType> class TOnlineAsyncOpHandle; }
 namespace UE::Online { template <typename OpType> class TOnlineResult; }
@@ -53,7 +54,7 @@ struct FGetUserInfo
 		/* User id to get info for*/
 		FAccountId AccountId;
 		/* Preferred platform for DisplayName*/
-		TOptional<FString> DisplayNamePlatform;
+		TOptional<EOnlinePlatformType> DisplayNamePlatform;
 	};
 
 	/**
diff --git a/Engine/Plugins/Online/OnlineServicesEOS/Source/Private/Online/OnlineServicesEOS.cpp b/Engine/Plugins/Online/OnlineServicesEOS/Source/Private/Online/OnlineServicesEOS.cpp
--- a/Engine/Plugins/Online/OnlineServicesEOS/Source/Private/Online/OnlineServicesEOS.cpp
+++ b/Engine/Plugins/Online/OnlineServicesEOS/Source/Private/Online/OnlineServicesEOS.cpp
@@ -32,4 +32,51 @@ void FOnlineServicesEOS::RegisterComponents()
 	Super::RegisterComponents();
 }
 
+EOS_OnlinePlatformType EOnlinePlatformType_To_EOS_OnlinePlatformType(const EOnlinePlatformType& InType)
+{
+	switch (InType)
+	{
+	case EOnlinePlatformType::Epic:		return EOS_OPT_Epic;
+	case EOnlinePlatformType::Steam:	return EOS_OPT_Steam;
+	case EOnlinePlatformType::PSN:		return EOS_OPT_PSN;
+	case EOnlinePlatformType::Nintendo:	return EOS_OPT_Nintendo;
+	case EOnlinePlatformType::XBL:		return EOS_OPT_XBL;
+	default:							checkNoEntry(); // Intentional fall-through
+	case EOnlinePlatformType::Unknown:	return EOS_OPT_Unknown;
+	}
+}
+
+EOnlinePlatformType EOS_OnlinePlatformType_To_EOnlinePlatformType(const EOS_OnlinePlatformType& InType)
+{
+	if (InType == EOS_OPT_Epic)
+	{
+		return EOnlinePlatformType::Epic;
+	}
+	else if (InType == EOS_OPT_Steam)
+	{
+		return EOnlinePlatformType::Steam;
+	}
+	else if (InType == EOS_OPT_PSN)
+	{
+		return EOnlinePlatformType::PSN;
+	}
+	else if (InType == EOS_OPT_Nintendo)
+	{
+		return EOnlinePlatformType::Nintendo;
+	}
+	else if (InType == EOS_OPT_XBL)
+	{
+		return EOnlinePlatformType::XBL;
+	}
+	else if (InType == EOS_OPT_Unknown)
+	{
+		return EOnlinePlatformType::Unknown;
+	}
+	else
+	{
+		checkNoEntry();
+		return EOnlinePlatformType::Unknown;
+	}
+}
+
 /* UE::Online */ }
diff --git a/Engine/Plugins/Online/OnlineServicesEOS/Source/Private/Online/PresenceEOS.cpp b/Engine/Plugins/Online/OnlineServicesEOS/Source/Private/Online/PresenceEOS.cpp
--- a/Engine/Plugins/Online/OnlineServicesEOS/Source/Private/Online/PresenceEOS.cpp
+++ b/Engine/Plugins/Online/OnlineServicesEOS/Source/Private/Online/PresenceEOS.cpp
@@ -9,6 +9,7 @@
 #include "Online/OnlineErrorEpicCommon.h"
 #include "Online/OnlineIdEOS.h"
 #include "Online/OnlineServicesEOS.h"
+#include "Templates/ValueOrError.h"
 
 #include "eos_presence.h"
 
@@ -196,6 +197,97 @@ TOnlineResult<FGetCachedPresence> FPresenceEOS::GetCachedPresence(FGetCachedPres
 	return TOnlineResult<FGetCachedPresence>(Errors::NotFound()); 
 }
 
+TValueOrError<FString, UE::Online::FOnlineError> PresenceProperty_To_EOS_Presence_String(const FPresenceProperty& PresenceAttribute)
+{
+	FString UpdatedPropertyValueStr;
+
+	if (PresenceAttribute.IsType<FPresencePropertiesRef>())
+	{
+		FPresencePropertiesRef UpdatedPropertyValueRef = PresenceAttribute.Get<FPresencePropertiesRef>();
+
+		UpdatedPropertyValueStr = FString::Printf(TEXT("m%s"), *NestedVariantToJson(UpdatedPropertyValueRef.ToSharedPtr()));
+	}
+	else if (PresenceAttribute.IsType<bool>())
+	{
+		UpdatedPropertyValueStr = FString::Printf(TEXT("b%s"), *ToLogString(PresenceAttribute.Get<bool>()));
+	}
+	else if (PresenceAttribute.IsType<int64>())
+	{
+		UpdatedPropertyValueStr = FString::Printf(TEXT("i%lld"), PresenceAttribute.Get<int64>());
+	}
+	else if (PresenceAttribute.IsType<double>())
+	{
+		UpdatedPropertyValueStr = FString::Printf(TEXT("d%f"), PresenceAttribute.Get<double>());
+	}
+	else if (PresenceAttribute.IsType<FString>())
+	{
+		UpdatedPropertyValueStr = FString::Printf(TEXT("s%s"), *PresenceAttribute.Get<FString>());
+	}
+	else
+	{
+		return MakeError(Errors::CantParse());
+	}
+
+	return MakeValue(MoveTemp(UpdatedPropertyValueStr));
+}
+
+TValueOrError<FPresenceProperty, UE::Online::FOnlineError> PresenceProperty_From_EOS_Presence_String(const char* EOSPresenceString, FString& OutRecordValueStr)
+{
+	FPresenceProperties::ValueType RecordValue;
+
+	if(EOSPresenceString != nullptr)
+	{
+		if (EOSPresenceString[0] == 'm')
+		{
+			const char* RecordValueUtf8 = EOSPresenceString + 1;
+					
+			OutRecordValueStr = UTF8_TO_TCHAR(RecordValueUtf8);
+
+			FPresencePropertiesRef MapRecordValue = FPresenceProperties::CreateVariant();
+			NestedVariantFromJson(RecordValueUtf8, MapRecordValue);
+
+			RecordValue.Set<FPresencePropertiesRef>(MapRecordValue);
+		}
+		else if (EOSPresenceString[0] == 'b')
+		{
+			const char* RecordValueUtf8 = EOSPresenceString + 1;
+
+			OutRecordValueStr = UTF8_TO_TCHAR(RecordValueUtf8);
+
+			RecordValue.Set<bool>(OutRecordValueStr.ToBool());
+		}
+		else if (EOSPresenceString[0] == 'i')
+		{
+			const char* RecordValueUtf8 = EOSPresenceString + 1;
+
+			RecordValue.Set<int64>(FCStringUtf8::Atoi64((const UTF8CHAR*)RecordValueUtf8));
+		}
+		else if (EOSPresenceString[0] == 'd')
+		{
+			const char* RecordValueUtf8 = EOSPresenceString + 1;
+
+			RecordValue.Set<double>(FCStringUtf8::Atod((const UTF8CHAR*)RecordValueUtf8));
+		}
+		else if (EOSPresenceString[0] == 's')
+		{
+			const char* RecordValueUtf8 = EOSPresenceString + 1;
+
+			OutRecordValueStr = UTF8_TO_TCHAR(RecordValueUtf8);
+
+			RecordValue.Set<FString>(OutRecordValueStr);
+		}
+		else
+		{
+			return MakeError(Errors::CantParse());
+		}
+	}
+	else
+	{
+		return MakeError(Errors::InvalidParams());
+	}
+
+	return MakeValue(MoveTemp(RecordValue));
+}
 
 TOnlineAsyncOpHandle<FUpdatePresence> FPresenceEOS::UpdatePresence(FUpdatePresence::Params&& InParams)
 {
@@ -210,6 +302,9 @@ TOnlineAsyncOpHandle<FUpdatePresence> FPresenceEOS::UpdatePresence(FUpdatePresen
 	Op->Then([this](TOnlineAsyncOp<FUpdatePresence>& InAsyncOp, TPromise<const EOS_Presence_SetPresenceCallbackInfo*>&& Promise) mutable
 	{
 		const FUpdatePresence::Params& Params = InAsyncOp.GetParams();
+		TSharedRef<FUserPresence> Presence = Params.Presence;
+		ModifyPresenceUpdate(Presence);
+
 		EOS_HPresenceModification ChangeHandle = nullptr;
 		EOS_Presence_CreatePresenceModificationOptions Options = { };
 		Options.ApiVersion = 1;
@@ -224,7 +319,7 @@ TOnlineAsyncOpHandle<FUpdatePresence> FPresenceEOS::UpdatePresence(FUpdatePresen
 			EOS_PresenceModification_SetStatusOptions StatusOptions = { };
 			StatusOptions.ApiVersion = 1;
 			UE_EOS_CHECK_API_MISMATCH(EOS_PRESENCEMODIFICATION_SETSTATUS_API_LATEST, 1);
-			StatusOptions.Status = ToEOS_Presence_EStatus(Params.Presence->Status);
+			StatusOptions.Status = ToEOS_Presence_EStatus(Presence->Status);
 			EOS_EResult SetStatusResult = EOS_PresenceModification_SetStatus(ChangeHandle, &StatusOptions);
 			if (SetStatusResult != EOS_EResult::EOS_Success)
 			{
@@ -240,7 +335,7 @@ TOnlineAsyncOpHandle<FUpdatePresence> FPresenceEOS::UpdatePresence(FUpdatePresen
 			RawRichTextOptions.ApiVersion = 1;
 			UE_EOS_CHECK_API_MISMATCH(EOS_PRESENCEMODIFICATION_SETRAWRICHTEXT_API_LATEST, 1);
 
-			const FTCHARToUTF8 Utf8RawRichText(*Params.Presence->StatusString);
+			const FTCHARToUTF8 Utf8RawRichText(*Presence->StatusString);
 			RawRichTextOptions.RichText = Utf8RawRichText.Get();
 
 			EOS_EResult SetRawRichTextResult = EOS_PresenceModification_SetRawRichText(ChangeHandle, &RawRichTextOptions);
@@ -255,9 +350,9 @@ TOnlineAsyncOpHandle<FUpdatePresence> FPresenceEOS::UpdatePresence(FUpdatePresen
 			// EOS needs to be specific on which fields are removed and which aren't, so grab the last presence and check which are deletions
 			TArray<FString> RemovedProperties;
 			TSharedRef<FUserPresence> LastUserPresence = FindOrCreatePresence(Params.LocalAccountId, Params.LocalAccountId);
-			for (const TPair<FString, FPresenceVariant>& Pair : LastUserPresence->Properties)
+			for (const TPair<FPresenceProperties::KeyType, FPresenceProperty>& Pair : LastUserPresence->Properties)
 			{
-				if(!Params.Presence->Properties.Contains(Pair.Key))
+				if(!Presence->Properties.Contains(Pair.Key))
 				{
 					RemovedProperties.Add(Pair.Key);
 				}
@@ -298,12 +393,12 @@ TOnlineAsyncOpHandle<FUpdatePresence> FPresenceEOS::UpdatePresence(FUpdatePresen
 			}
 
 			// Added/Updated fields
-			if (Params.Presence->Properties.Num() > 0)
+			if (Presence->Properties.Num() > 0)
 			{
-				if (Params.Presence->Properties.Num() > EOS_PRESENCE_DATA_MAX_KEYS)
+				if (Presence->Properties.Num() > EOS_PRESENCE_DATA_MAX_KEYS)
 				{
 					// TODO: Move this check higher.  Needs to take into account number of present fields (not just ones updated) and removed fields.
-					UE_LOG(LogOnlineServices, Warning, TEXT("UpdatePresence: Too many presence keys.  %u/%u"), Params.Presence->Properties.Num(), EOS_PRESENCE_DATA_MAX_KEYS);
+					UE_LOG(LogOnlineServices, Warning, TEXT("UpdatePresence: Too many presence keys.  %u/%u"), Presence->Properties.Num(), EOS_PRESENCE_DATA_MAX_KEYS);
 					InAsyncOp.SetError(Errors::InvalidParams());
 					Promise.EmplaceValue();
 					return;
@@ -311,17 +406,29 @@ TOnlineAsyncOpHandle<FUpdatePresence> FPresenceEOS::UpdatePresence(FUpdatePresen
 				TArray<FTCHARToUTF8, TInlineAllocator<EOS_PRESENCE_DATA_MAX_KEYS * 2>> Utf8Strings;
 				TArray<EOS_Presence_DataRecord, TInlineAllocator<EOS_PRESENCE_DATA_MAX_KEYS>> Records;
 				int32 CurrentIndex = 0;
-				for (const TPair<FString, FPresenceVariant>& UpdatedProperty : Params.Presence->Properties)
+				for (const TPair<FPresenceProperties::KeyType, FPresenceProperty>& UpdatedProperty : Presence->Properties)
 				{
 					const FTCHARToUTF8& Utf8Key = Utf8Strings.Emplace_GetRef(*UpdatedProperty.Key);
-					const FTCHARToUTF8& Utf8Value = Utf8Strings.Emplace_GetRef(*UpdatedProperty.Value); // TODO: Better serialization
-
-					EOS_Presence_DataRecord& Record = Records.Emplace_GetRef();
-					Record.ApiVersion = 1;
-					UE_EOS_CHECK_API_MISMATCH(EOS_PRESENCE_DATARECORD_API_LATEST, 1);
-					Record.Key = Utf8Key.Get();
-					Record.Value = Utf8Value.Get();
-					UE_LOG(LogOnlineServices, Warning, TEXT("UpdatePresence: Set field [%s] to [%s]"), *UpdatedProperty.Key, *UpdatedProperty.Value); // Temp logging
+
+					TValueOrError<FString, FOnlineError> UpdatedPropertyValueResult = PresenceProperty_To_EOS_Presence_String(UpdatedProperty.Value);
+					if (UpdatedPropertyValueResult.HasValue())
+					{
+						const FTCHARToUTF8& Utf8Value = Utf8Strings.Emplace_GetRef(*UpdatedPropertyValueResult.GetValue());
+							
+						EOS_Presence_DataRecord& Record = Records.Emplace_GetRef();
+						Record.ApiVersion = 1;
+						UE_EOS_CHECK_API_MISMATCH(EOS_PRESENCE_DATARECORD_API_LATEST, 1);
+						Record.Key = Utf8Key.Get();
+						Record.Value = Utf8Value.Get();
+						UE_LOG(LogOnlineServices, VeryVerbose, TEXT("PartialUpdatePresence: Set field [%s] to [%s]"), *UpdatedProperty.Key, *UpdatedPropertyValueResult.GetValue());
+					}
+					else
+					{
+						UE_LOG(LogOnlineServices, Warning, TEXT("PartialUpdatePresence: Unsupported data type passed for presence attribute [%s]"), *UpdatedProperty.Key);
+						InAsyncOp.SetError(Errors::InvalidParams());
+						Promise.EmplaceValue();
+						return;
+					}
 				}
 
 				EOS_PresenceModification_SetDataOptions DataOptions = { };
@@ -363,8 +470,11 @@ TOnlineAsyncOpHandle<FUpdatePresence> FPresenceEOS::UpdatePresence(FUpdatePresen
 		{
 			// Update local presence
 			const FUpdatePresence::Params& Params = InAsyncOp.GetParams();
+			TSharedRef<FUserPresence> Presence = Params.Presence;
+			ModifyPresenceUpdate(Presence);
+
 			TSharedRef<FUserPresence> LocalUserPresence = FindOrCreatePresence(Params.LocalAccountId, Params.LocalAccountId);
-			*LocalUserPresence = *Params.Presence;
+			*LocalUserPresence = *Presence;
 
 			InAsyncOp.SetResult(FUpdatePresence::Result());
 		}
@@ -401,6 +511,9 @@ TOnlineAsyncOpHandle<FPartialUpdatePresence> FPresenceEOS::PartialUpdatePresence
 		Op->Then([this](TOnlineAsyncOp<FPartialUpdatePresence>& InAsyncOp, TPromise<const EOS_Presence_SetPresenceCallbackInfo*>&& Promise) mutable
 			{
 				const FPartialUpdatePresence::Params& Params = InAsyncOp.GetParams();
+				FPartialUpdatePresence::Params::FMutations Mutations = Params.Mutations;
+				ModifyPartialPresenceUpdate(Mutations);
+
 				EOS_HPresenceModification ChangeHandle = nullptr;
 				EOS_Presence_CreatePresenceModificationOptions Options = { };
 				Options.ApiVersion = 1;
@@ -412,12 +525,12 @@ TOnlineAsyncOpHandle<FPartialUpdatePresence> FPresenceEOS::PartialUpdatePresence
 					check(ChangeHandle != nullptr);
 
 					// State
-					if (Params.Mutations.Status.IsSet())
+					if (Mutations.Status.IsSet())
 					{
 						EOS_PresenceModification_SetStatusOptions StatusOptions = { };
 						StatusOptions.ApiVersion = 1;
 						UE_EOS_CHECK_API_MISMATCH(EOS_PRESENCEMODIFICATION_SETSTATUS_API_LATEST, 1);
-						StatusOptions.Status = ToEOS_Presence_EStatus(Params.Mutations.Status.GetValue());
+						StatusOptions.Status = ToEOS_Presence_EStatus(Mutations.Status.GetValue());
 						EOS_EResult SetStatusResult = EOS_PresenceModification_SetStatus(ChangeHandle, &StatusOptions);
 						if (SetStatusResult != EOS_EResult::EOS_Success)
 						{
@@ -428,19 +541,19 @@ TOnlineAsyncOpHandle<FPartialUpdatePresence> FPresenceEOS::PartialUpdatePresence
 						}
 						else
 						{
-							UE_LOG(LogOnlineServices, Warning, TEXT("UpdatePresence: Status set to %u"), (uint8)Params.Mutations.Status.GetValue()); // Temp logging
+							UE_LOG(LogOnlineServices, Warning, TEXT("UpdatePresence: Status set to %u"), (uint8)Mutations.Status.GetValue()); // Temp logging
 						}
 					}
 
 					// Raw rich text
-					if (Params.Mutations.StatusString.IsSet())
+					if (Mutations.StatusString.IsSet())
 					{
 						// Convert the status string as the rich text string
 						EOS_PresenceModification_SetRawRichTextOptions RawRichTextOptions = { };
 						RawRichTextOptions.ApiVersion = 1;
 						UE_EOS_CHECK_API_MISMATCH(EOS_PRESENCEMODIFICATION_SETRAWRICHTEXT_API_LATEST, 1);
 
-						const FTCHARToUTF8 Utf8RawRichText(*Params.Mutations.StatusString.GetValue());
+						const FTCHARToUTF8 Utf8RawRichText(*Mutations.StatusString.GetValue());
 						RawRichTextOptions.RichText = Utf8RawRichText.Get();
 
 						EOS_EResult SetRawRichTextResult = EOS_PresenceModification_SetRawRichText(ChangeHandle, &RawRichTextOptions);
@@ -453,17 +566,17 @@ TOnlineAsyncOpHandle<FPartialUpdatePresence> FPresenceEOS::PartialUpdatePresence
 						}
 						else
 						{
-							UE_LOG(LogOnlineServices, Warning, TEXT("UpdatePresence: RichText set to %s"), *Params.Mutations.StatusString.GetValue()); // Temp logging
+							UE_LOG(LogOnlineServices, Warning, TEXT("UpdatePresence: RichText set to %s"), *Mutations.StatusString.GetValue()); // Temp logging
 						}
 					}
 					// Removed fields
-					if (Params.Mutations.RemovedProperties.Num() > 0)
+					if (Mutations.RemovedProperties.Num() > 0)
 					{
 						// EOS_PresenceModification_DeleteData
 						TArray<FTCHARToUTF8, TInlineAllocator<EOS_PRESENCE_DATA_MAX_KEYS>> Utf8Strings;
 						TArray<EOS_PresenceModification_DataRecordId, TInlineAllocator<EOS_PRESENCE_DATA_MAX_KEYS>> RecordIds;
 						int32 CurrentIndex = 0;
-						for (const FString& RemovedProperty : Params.Mutations.RemovedProperties)
+						for (const FString& RemovedProperty : Mutations.RemovedProperties)
 						{
 							const FTCHARToUTF8& Utf8Key = Utf8Strings.Emplace_GetRef(*RemovedProperty);
 
@@ -490,48 +603,58 @@ TOnlineAsyncOpHandle<FPartialUpdatePresence> FPresenceEOS::PartialUpdatePresence
 						}
 					}
 
+					if (Mutations.UpdatedProperties.Num() > EOS_PRESENCE_DATA_MAX_KEYS)
+					{
+						// TODO: Move this check higher.  Needs to take into account number of present fields (not just ones updated) and removed fields.
+						UE_LOG(LogOnlineServices, Warning, TEXT("UpdatePresence: Too many presence keys.  %u/%u"), Mutations.UpdatedProperties.Num(), EOS_PRESENCE_DATA_MAX_KEYS);
+						InAsyncOp.SetError(Errors::InvalidParams());
+						Promise.EmplaceValue();
+						return;
+					}
+
 					// Added/Updated fields
-					if (Params.Mutations.UpdatedProperties.Num() > 0)
+					TArray<FTCHARToUTF8, TInlineAllocator<EOS_PRESENCE_DATA_MAX_KEYS * 2>> Utf8Strings;
+					TArray<EOS_Presence_DataRecord, TInlineAllocator<EOS_PRESENCE_DATA_MAX_KEYS>> Records;
+					
+					for (const TPair<FPresenceProperties::KeyType, FPresenceProperty>& UpdatedProperty : Mutations.UpdatedProperties)
 					{
-						if (Params.Mutations.UpdatedProperties.Num() > EOS_PRESENCE_DATA_MAX_KEYS)
-						{
-							// TODO: Move this check higher.  Needs to take into account number of present fields (not just ones updated) and removed fields.
-							UE_LOG(LogOnlineServices, Warning, TEXT("UpdatePresence: Too many presence keys.  %u/%u"), Params.Mutations.UpdatedProperties.Num(), EOS_PRESENCE_DATA_MAX_KEYS);
-							InAsyncOp.SetError(Errors::InvalidParams());
-							Promise.EmplaceValue();
-							return;
-						}
-						TArray<FTCHARToUTF8, TInlineAllocator<EOS_PRESENCE_DATA_MAX_KEYS * 2>> Utf8Strings;
-						TArray<EOS_Presence_DataRecord, TInlineAllocator<EOS_PRESENCE_DATA_MAX_KEYS>> Records;
-						int32 CurrentIndex = 0;
-						for (const TPair<FString, FPresenceVariant>& UpdatedProperty : Params.Mutations.UpdatedProperties)
-						{
-							const FTCHARToUTF8& Utf8Key = Utf8Strings.Emplace_GetRef(*UpdatedProperty.Key);
-							const FTCHARToUTF8& Utf8Value = Utf8Strings.Emplace_GetRef(*UpdatedProperty.Value); // TODO: Better serialization
+						const FTCHARToUTF8& Utf8Key = Utf8Strings.Emplace_GetRef(*UpdatedProperty.Key);
 
+						TValueOrError<FString, FOnlineError> UpdatedPropertyValueResult = PresenceProperty_To_EOS_Presence_String(UpdatedProperty.Value);
+						if (UpdatedPropertyValueResult.HasValue())
+						{
+							const FTCHARToUTF8& Utf8Value = Utf8Strings.Emplace_GetRef(*UpdatedPropertyValueResult.GetValue());
+							
 							EOS_Presence_DataRecord& Record = Records.Emplace_GetRef();
 							Record.ApiVersion = 1;
 							UE_EOS_CHECK_API_MISMATCH(EOS_PRESENCE_DATARECORD_API_LATEST, 1);
 							Record.Key = Utf8Key.Get();
 							Record.Value = Utf8Value.Get();
-							UE_LOG(LogOnlineServices, Warning, TEXT("UpdatePresence: Set field [%s] to [%s]"), *UpdatedProperty.Key, *UpdatedProperty.Value); // Temp logging
+							UE_LOG(LogOnlineServices, VeryVerbose, TEXT("PartialUpdatePresence: Set field [%s] to [%s]"), *UpdatedProperty.Key, *UpdatedPropertyValueResult.GetValue());
 						}
-
-						EOS_PresenceModification_SetDataOptions DataOptions = { };
-						DataOptions.ApiVersion = 1;
-						UE_EOS_CHECK_API_MISMATCH(EOS_PRESENCEMODIFICATION_SETDATA_API_LATEST, 1);
-						DataOptions.RecordsCount = Records.Num();
-						DataOptions.Records = Records.GetData();
-						EOS_EResult SetDataResult = EOS_PresenceModification_SetData(ChangeHandle, &DataOptions);
-						if (SetDataResult != EOS_EResult::EOS_Success)
+						else
 						{
-							UE_LOG(LogOnlineServices, Warning, TEXT("UpdatePresence: EOS_PresenceModification_SetData failed with result %s"), *LexToString(SetDataResult));
-							InAsyncOp.SetError(Errors::FromEOSResult(SetDataResult));
+							UE_LOG(LogOnlineServices, Warning, TEXT("PartialUpdatePresence: Unsupported data type passed for presence attribute [%s]"), *UpdatedProperty.Key);
+							InAsyncOp.SetError(Errors::InvalidParams());
 							Promise.EmplaceValue();
 							return;
 						}
 					}
 
+					EOS_PresenceModification_SetDataOptions DataOptions = { };
+					DataOptions.ApiVersion = 1;
+					UE_EOS_CHECK_API_MISMATCH(EOS_PRESENCEMODIFICATION_SETDATA_API_LATEST, 1);
+					DataOptions.RecordsCount = Records.Num();
+					DataOptions.Records = Records.GetData();
+					EOS_EResult SetDataResult = EOS_PresenceModification_SetData(ChangeHandle, &DataOptions);
+					if (SetDataResult != EOS_EResult::EOS_Success)
+					{
+						UE_LOG(LogOnlineServices, Warning, TEXT("PartialUpdatePresence: EOS_PresenceModification_SetData failed with result %s"), *LexToString(SetDataResult));
+						InAsyncOp.SetError(Errors::FromEOSResult(SetDataResult));
+						Promise.EmplaceValue();
+						return;
+					}
+
 					EOS_Presence_SetPresenceOptions SetPresenceOptions = { };
 					SetPresenceOptions.ApiVersion = 1;
 					UE_EOS_CHECK_API_MISMATCH(EOS_PRESENCE_SETPRESENCE_API_LATEST, 1);
@@ -555,34 +678,38 @@ TOnlineAsyncOpHandle<FPartialUpdatePresence> FPresenceEOS::PartialUpdatePresence
 				{
 					// Update local presence
 					bool bPresenceHasChanged = false;
+
 					const FPartialUpdatePresence::Params& Params = InAsyncOp.GetParams();
+					FPartialUpdatePresence::Params::FMutations Mutations = Params.Mutations;
+					ModifyPartialPresenceUpdate(Mutations);
+
 					TSharedRef<FUserPresence> LocalUserPresence = FindOrCreatePresence(Params.LocalAccountId, Params.LocalAccountId);
-					if (Params.Mutations.Status.IsSet())
+					if (Mutations.Status.IsSet())
 					{
-						if (LocalUserPresence->Status != Params.Mutations.Status.GetValue())
+						if (LocalUserPresence->Status != Mutations.Status.GetValue())
 						{
 							bPresenceHasChanged = true;
 						}
 
-						LocalUserPresence->Status = Params.Mutations.Status.GetValue();
+						LocalUserPresence->Status = Mutations.Status.GetValue();
 					}
-					if (Params.Mutations.StatusString.IsSet())
+					if (Mutations.StatusString.IsSet())
 					{
-						LocalUserPresence->StatusString = Params.Mutations.StatusString.GetValue();
+						LocalUserPresence->StatusString = Mutations.StatusString.GetValue();
 					}
-					if (Params.Mutations.GameStatus.IsSet())
+					if (Mutations.GameStatus.IsSet())
 					{
-						LocalUserPresence->GameStatus = Params.Mutations.GameStatus.GetValue();
+						LocalUserPresence->GameStatus = Mutations.GameStatus.GetValue();
 					}
-					if (Params.Mutations.Joinability.IsSet())
+					if (Mutations.Joinability.IsSet())
 					{
-						LocalUserPresence->Joinability = Params.Mutations.Joinability.GetValue();
+						LocalUserPresence->Joinability = Mutations.Joinability.GetValue();
 					}
-					for (const FString& RemovedKey : Params.Mutations.RemovedProperties)
+					for (const FString& RemovedKey : Mutations.RemovedProperties)
 					{
 						LocalUserPresence->Properties.Remove(RemovedKey);
 					}
-					for (const TPair<FString, FPresenceVariant>& UpdatedProperty : Params.Mutations.UpdatedProperties)
+					for (const TPair<FPresenceProperties::KeyType, FPresenceProperty>& UpdatedProperty : Mutations.UpdatedProperties)
 					{
 						LocalUserPresence->Properties.Emplace(UpdatedProperty.Key, UpdatedProperty.Value);
 					}
@@ -620,6 +747,39 @@ TSharedRef<FUserPresence> FPresenceEOS::FindOrCreatePresence(FAccountId LocalAcc
 	return UserPresence;
 }
 
+EOnlinePlatformType IntegratedPlatform_To_OnlinePlatformType(const FString& IntegratedPlatformStr)
+{
+	if (IntegratedPlatformStr == TEXT("EGS"))
+	{
+		return EOnlinePlatformType::Epic;
+	}
+	else if (IntegratedPlatformStr == TEXT("STEAM"))
+	{
+		return EOnlinePlatformType::Steam;
+	}
+	else if (IntegratedPlatformStr == TEXT("PSN"))
+	{
+		return EOnlinePlatformType::PSN;
+	}
+	else if (IntegratedPlatformStr == TEXT("NINTENDO"))
+	{
+		return EOnlinePlatformType::Nintendo;
+	}
+	else if (IntegratedPlatformStr == TEXT("XBL"))
+	{
+		return EOnlinePlatformType::XBL;
+	}
+	else if (IntegratedPlatformStr == TEXT("UNKNOWN"))
+	{
+		return EOnlinePlatformType::Unknown;
+	}
+	else 
+	{
+		checkNoEntry();
+		return EOnlinePlatformType::Unknown;
+	}
+}
+
 void FPresenceEOS::UpdateUserPresence(FAccountId LocalAccountId, FAccountId PresenceAccountId)
 {
 	const EOS_EpicAccountId LocalUserId = GetEpicAccountId(LocalAccountId);
@@ -669,24 +829,52 @@ void FPresenceEOS::UpdateUserPresence(FAccountId LocalAccountId, FAccountId Pres
 		{
 			// TODO:  Handle Properties that aren't replicated through presence (eg "ProductId")
 			TArrayView<const EOS_Presence_DataRecord> Records(PresenceInfo->Records, PresenceInfo->RecordsCount);
+
 			// Detect removals
 			TArray<FString, TInlineAllocator<EOS_PRESENCE_DATA_MAX_KEYS>> RemovedKeys;
 			UserPresence->Properties.GenerateKeyArray(RemovedKeys);
+
 			for (const EOS_Presence_DataRecord& Record : Records)
 			{
 				FString RecordKey = UTF8_TO_TCHAR(Record.Key);
-				FString RecordValue = UTF8_TO_TCHAR(Record.Value);
-				RemovedKeys.Remove(RecordKey);
-				if (FString* ExistingValue = UserPresence->Properties.Find(RecordKey))
+
+				FString RecordValueStr;
+				TValueOrError<FPresenceProperty, FOnlineError> RecordValueResult = PresenceProperty_From_EOS_Presence_String(Record.Value, RecordValueStr);
+				if (RecordValueResult.HasValue())
 				{
-					if (*ExistingValue == RecordValue)
+					FPresenceProperty RecordValue = RecordValueResult.GetValue();
+
+					RemovedKeys.Remove(RecordKey);
+					if (FPresenceProperties::ValueType* ExistingValue = UserPresence->Properties.Find(RecordKey))
 					{
-						continue; // No change
+						if (ExistingValue->IsType<FPresencePropertiesRef>() && RecordValue.IsType<FPresencePropertiesRef>())
+						{
+							const FString ExistingValueJsonStr = NestedVariantToJson(ExistingValue->Get<FPresencePropertiesRef>().ToSharedPtr());
+
+							if (ExistingValueJsonStr == RecordValueStr)
+							{
+								continue; // No change
+							}
+						}
+						else if ((ExistingValue->IsType<int64>() && RecordValue.IsType<int64>() && ExistingValue->Get<int64>() == RecordValue.Get<int64>()) ||
+								(ExistingValue->IsType<long long>() && RecordValue.IsType<long long>() && ExistingValue->Get<long long>() == RecordValue.Get<long long>())) // Setting a variant as int64 can appear as long long too
+						{
+							continue; // No change
+						}
+						else if (ExistingValue->IsType<FString>() && RecordValue.IsType<FString>() && ExistingValue->Get<FString>() == RecordValue.Get<FString>())
+						{
+							continue; // No change
+						}
 					}
+					bPresenceHasChanged = true;
+					UserPresence->Properties.Add(MoveTemp(RecordKey), MoveTemp(RecordValue));
+				}
+				else
+				{
+					UE_LOG(LogOnlineServices, Warning, TEXT("UpdateUserPresence: Unsupported data type passed for presence attribute [%s]"), *RecordKey);
 				}
-				bPresenceHasChanged = true;
-				UserPresence->Properties.Add(MoveTemp(RecordKey), MoveTemp(RecordValue));
 			}
+
 			// Any fields that have been removed
 			if (RemovedKeys.Num() > 0)
 			{
@@ -702,11 +890,24 @@ void FPresenceEOS::UpdateUserPresence(FAccountId LocalAccountId, FAccountId Pres
 			bPresenceHasChanged = true;
 			UserPresence->Properties.Reset();
 		}
+
+		// We set the PlatformType by translating the IntegratedPlatform field
+		FString IntegratedPlatformStr = UTF8_TO_TCHAR(PresenceInfo->IntegratedPlatform);
+		if (!IntegratedPlatformStr.IsEmpty())
+		{
+			EOnlinePlatformType OnlinePlatformType = IntegratedPlatform_To_OnlinePlatformType(IntegratedPlatformStr);
+			if (UserPresence->PlatformType != OnlinePlatformType)
+			{
+				bPresenceHasChanged = true;
+				UserPresence->PlatformType = OnlinePlatformType;
+			}
+		}
+
 		EOS_Presence_Info_Release(PresenceInfo);
 	}
 	else
 	{
-		UE_LOG(LogOnlineServices, Warning, TEXT("UpdateUserPresence: CopyPresence Failed %s"), *LexToString(CopyPresenceResult));
+		UE_LOG(LogOnlineServices, Warning, TEXT("UpdateUserPresence: CopyPresence failed with result: %s"), *LexToString(CopyPresenceResult));
 	}
 
 	if (bPresenceHasChanged)
@@ -744,4 +945,14 @@ FAccountId FPresenceEOS::FindAccountId(const EOS_EpicAccountId EpicAccountId)
 	return UE::Online::FindAccountId(Services.GetServicesProvider(), EpicAccountId);
 }
 
+void FPresenceEOS::ModifyPresenceUpdate(TSharedRef<FUserPresence>& Presence)
+{
+	// To be implemented in derived classes
+}
+
+void FPresenceEOS::ModifyPartialPresenceUpdate(FPartialUpdatePresence::Params::FMutations& Mutations)
+{
+	// To be implemented in derived classes
+}
+
 /* UE::Online */ }
diff --git a/Engine/Plugins/Online/OnlineServicesEOS/Source/Private/Online/UserInfoEOS.cpp b/Engine/Plugins/Online/OnlineServicesEOS/Source/Private/Online/UserInfoEOS.cpp
--- a/Engine/Plugins/Online/OnlineServicesEOS/Source/Private/Online/UserInfoEOS.cpp
+++ b/Engine/Plugins/Online/OnlineServicesEOS/Source/Private/Online/UserInfoEOS.cpp
@@ -160,15 +160,14 @@ TOnlineResult<FGetUserInfo> FUserInfoEOS::GetUserInfo(FGetUserInfo::Params&& Par
 		EOS_OnlinePlatformType OnlinePlatformType = EOS_OPT_Epic;
 		if (Params.DisplayNamePlatform.IsSet())
 		{
-			const FString& DisplayNamePlatform = Params.DisplayNamePlatform.GetValue();
-			EOS_OnlinePlatformType NewOnlinePlatformType = EOSOnlinePlatformTypeFromString(DisplayNamePlatform);
+			EOS_OnlinePlatformType NewOnlinePlatformType = EOnlinePlatformType_To_EOS_OnlinePlatformType(Params.DisplayNamePlatform.GetValue());
 			if (NewOnlinePlatformType != EOS_OPT_Unknown)
 			{
 				OnlinePlatformType = NewOnlinePlatformType;
 			}
 			else
 			{
-				UE_LOG(LogOnlineServices, VeryVerbose, TEXT("FGetUserInfo::Params::DisplayNamePlatform set to %s. Not a valid Online Platform Type. Epic will be used as default."), *DisplayNamePlatform);
+				UE_LOG(LogOnlineServices, VeryVerbose, TEXT("FGetUserInfo::Params::DisplayNamePlatform platform type unknown. Epic will be used as default."));
 			}
 		}
 		
diff --git a/Engine/Plugins/Online/OnlineServicesEOS/Source/Public/Online/OnlineServicesEOS.h b/Engine/Plugins/Online/OnlineServicesEOS/Source/Public/Online/OnlineServicesEOS.h
--- a/Engine/Plugins/Online/OnlineServicesEOS/Source/Public/Online/OnlineServicesEOS.h
+++ b/Engine/Plugins/Online/OnlineServicesEOS/Source/Public/Online/OnlineServicesEOS.h
@@ -28,4 +28,7 @@ public:
 	ONLINESERVICESEOS_API virtual void RegisterComponents() override;
 };
 
+ONLINESERVICESEOS_API EOS_OnlinePlatformType EOnlinePlatformType_To_EOS_OnlinePlatformType(const EOnlinePlatformType& InType);
+ONLINESERVICESEOS_API EOnlinePlatformType EOS_OnlinePlatformType_To_EOnlinePlatformType(const EOS_OnlinePlatformType& InType);
+
 /* UE::Online */ }
diff --git a/Engine/Plugins/Online/OnlineServicesEOS/Source/Public/Online/PresenceEOS.h b/Engine/Plugins/Online/OnlineServicesEOS/Source/Public/Online/PresenceEOS.h
--- a/Engine/Plugins/Online/OnlineServicesEOS/Source/Public/Online/PresenceEOS.h
+++ b/Engine/Plugins/Online/OnlineServicesEOS/Source/Public/Online/PresenceEOS.h
@@ -42,6 +42,11 @@ protected:
 
 	ONLINESERVICESEOS_API FAccountId FindAccountId(const EOS_EpicAccountId EpicAccountId);
 
+	/** Allow derived classes to modify the content of presence updates */
+	ONLINESERVICESEOS_API virtual void ModifyPresenceUpdate(TSharedRef<FUserPresence>& Presence);
+	/** Allow derived classes to modify the content of partial presence updates */
+	ONLINESERVICESEOS_API virtual void ModifyPartialPresenceUpdate(FPartialUpdatePresence::Params::FMutations& Mutations);
+
 	EOS_HPresence PresenceHandle = nullptr;
 
 	/** Login status changed event handle */
diff --git a/Engine/Plugins/Online/OnlineServicesOSSAdapter/Source/Private/Online/PresenceOSSAdapter.cpp b/Engine/Plugins/Online/OnlineServicesOSSAdapter/Source/Private/Online/PresenceOSSAdapter.cpp
--- a/Engine/Plugins/Online/OnlineServicesOSSAdapter/Source/Private/Online/PresenceOSSAdapter.cpp
+++ b/Engine/Plugins/Online/OnlineServicesOSSAdapter/Source/Private/Online/PresenceOSSAdapter.cpp
@@ -68,20 +68,24 @@ EUserPresenceGameStatus GameStatusV1toV2(FOnlineUserPresence& Presence)
 	}
 }
 
-void PropertiesV1toV2(FOnlineKeyValuePairs<FPresenceKey, FVariantData>& PropertiesIn, TMap<FString, FPresenceVariant>& PropertiesOut)
+void PropertiesV1toV2(FOnlineKeyValuePairs<FPresenceKey, FVariantData>& PropertiesIn, TMap<FString, FPresenceProperty>& PropertiesOut)
 {
 	for (const TPair<FPresenceKey, FVariantData>& Pair : PropertiesIn)
 	{
-		PropertiesOut.Add(Pair.Key, Pair.Value.ToString());
+		FPresenceProperty Value;
+
+		Value.Set<FString>(Pair.Value.ToString()); // TODO: Proper conversion from V2 Variants to V1
+
+		PropertiesOut.Add(Pair.Key, Value);
 	}
 }
 
-void PropertiesV2toV1(TMap<FString, FPresenceVariant> PropertiesIn, FOnlineKeyValuePairs<FPresenceKey, FVariantData>& PropertiesOut)
+void PropertiesV2toV1(TMap<FString, FPresenceProperty> PropertiesIn, FOnlineKeyValuePairs<FPresenceKey, FVariantData>& PropertiesOut)
 {
-	for (const TPair<FString, FPresenceVariant>& Pair : PropertiesIn)
+	for (const TPair<FString, FPresenceProperty>& Pair : PropertiesIn)
 	{
 		FVariantData Data;
-		Data.SetValue(Pair.Value);
+		Data.SetValue(Pair.Value.Get<FString>()); // TODO: Proper conversion from V2 Variants to V1
 		PropertiesOut.Add(Pair.Key, Data);
 	}
 }
diff --git a/Engine/Source/Runtime/CoreOnline/Private/Online/NestedVariantJson.cpp b/Engine/Source/Runtime/CoreOnline/Private/Online/NestedVariantJson.cpp
--- a/Engine/Source/Runtime/CoreOnline/Private/Online/NestedVariantJson.cpp
+++ b/Engine/Source/Runtime/CoreOnline/Private/Online/NestedVariantJson.cpp
@@ -184,3 +184,11 @@ void NestedVariantFromJsonArray(const TArray<TSharedPtr<FJsonValue>>& JsonArray,
 		}
 	}
 }
+
+void LexFromString(FNestedVariantJson::FMap& OutValue, const TCHAR* InString)
+{
+	FNestedVariantJson::FMapRef  OutValueRef = FNestedVariantJson::FMap::CreateVariant();
+	NestedVariantFromJson(TCHAR_TO_UTF8(InString), OutValueRef);
+
+	OutValue = OutValueRef.Get();
+}
\ No newline at end of file
diff --git a/Engine/Source/Runtime/CoreOnline/Public/Online/NestedVariantJson.h b/Engine/Source/Runtime/CoreOnline/Public/Online/NestedVariantJson.h
--- a/Engine/Source/Runtime/CoreOnline/Public/Online/NestedVariantJson.h
+++ b/Engine/Source/Runtime/CoreOnline/Public/Online/NestedVariantJson.h
@@ -18,3 +18,5 @@ COREONLINE_API TSharedPtr<FJsonValue> NestedVariantToJsonValue(const FNestedVari
 COREONLINE_API void NestedVariantFromJson(const char* InJson, FNestedVariantJson::FMapRef& Map);
 COREONLINE_API void NestedVariantFromJsonObject(const TSharedPtr<FJsonObject>& JsonObject, FNestedVariantJson::FMapRef& Map);
 COREONLINE_API void NestedVariantFromJsonArray(const TArray<TSharedPtr<FJsonValue>>& JsonArray, FNestedVariantJson::FArrayRef& Array);
+
+COREONLINE_API void LexFromString(FNestedVariantJson::FMap& OutValue, const TCHAR* InString);
\ No newline at end of file
