Minor bugfixes for systems touched in CLs 45098488-45101241

diff --git a/Engine/Plugins/Online/OnlineFramework/Source/Party/Private/User/SocialUser.cpp b/Engine/Plugins/Online/OnlineFramework/Source/Party/Private/User/SocialUser.cpp
--- a/Engine/Plugins/Online/OnlineFramework/Source/Party/Private/User/SocialUser.cpp
+++ b/Engine/Plugins/Online/OnlineFramework/Source/Party/Private/User/SocialUser.cpp
@@ -1342,7 +1342,10 @@ bool USocialUser::IsOnline() const
 						TOnlineResult<FGetCachedPresence> GetCachedPresenceResult = Presence->GetCachedPresence({ LocalUserAccountId, TargetAccountId });
 						if (GetCachedPresenceResult.IsOk())
 						{
-							return GetCachedPresenceResult.GetOkValue().Presence->Status == EUserPresenceStatus::Online;
+							return GetCachedPresenceResult.GetOkValue().Presence->Status == EUserPresenceStatus::Online
+								|| GetCachedPresenceResult.GetOkValue().Presence->Status == EUserPresenceStatus::Away
+								|| GetCachedPresenceResult.GetOkValue().Presence->Status == EUserPresenceStatus::ExtendedAway
+								|| GetCachedPresenceResult.GetOkValue().Presence->Status == EUserPresenceStatus::DoNotDisturb;
 						}
 					}
 				}
@@ -1401,7 +1404,7 @@ bool USocialUser::IsPlayingThisGame() const
 						if (GetCachedPresenceResult.IsOk())
 						{
 							EUserPresenceGameStatus CachedGameStatus = GetCachedPresenceResult.GetOkValue().Presence->GameStatus;
-							return CachedGameStatus == EUserPresenceGameStatus::PlayingThisGame || CachedGameStatus == UE::Online::EUserPresenceGameStatus::PlayingOtherGame;
+							return CachedGameStatus == EUserPresenceGameStatus::PlayingThisGame;
 						}
 					}
 				}
diff --git a/Engine/Plugins/Online/OnlineServicesEOS/Source/Private/Online/PresenceEOS.cpp b/Engine/Plugins/Online/OnlineServicesEOS/Source/Private/Online/PresenceEOS.cpp
--- a/Engine/Plugins/Online/OnlineServicesEOS/Source/Private/Online/PresenceEOS.cpp
+++ b/Engine/Plugins/Online/OnlineServicesEOS/Source/Private/Online/PresenceEOS.cpp
@@ -696,22 +696,50 @@ TOnlineAsyncOpHandle<FPartialUpdatePresence> FPresenceEOS::PartialUpdatePresence
 
 						LocalUserPresence->Status = Mutations.Status.GetValue();
 					}
+
 					if (Mutations.StatusString.IsSet())
 					{
+						if (LocalUserPresence->StatusString != Mutations.StatusString.GetValue())
+						{
+							bPresenceHasChanged = true;
+						}
+
 						LocalUserPresence->StatusString = Mutations.StatusString.GetValue();
 					}
+
 					if (Mutations.GameStatus.IsSet())
 					{
+						if (LocalUserPresence->GameStatus != Mutations.GameStatus.GetValue())
+						{
+							bPresenceHasChanged = true;
+						}
+
 						LocalUserPresence->GameStatus = Mutations.GameStatus.GetValue();
 					}
+
 					if (Mutations.Joinability.IsSet())
 					{
+						if (LocalUserPresence->Joinability != Mutations.Joinability.GetValue())
+						{
+							bPresenceHasChanged = true;
+						}
+
 						LocalUserPresence->Joinability = Mutations.Joinability.GetValue();
 					}
+
+					if (!Mutations.RemovedProperties.IsEmpty())
+					{
+						bPresenceHasChanged = true;
+					}
 					for (const FString& RemovedKey : Mutations.RemovedProperties)
 					{
 						LocalUserPresence->Properties.Remove(RemovedKey);
 					}
+
+					if (!Mutations.UpdatedProperties.IsEmpty())
+					{
+						bPresenceHasChanged = true;
+					}
 					for (const TPair<FPresenceProperties::KeyType, FPresenceProperty>& UpdatedProperty : Mutations.UpdatedProperties)
 					{
 						LocalUserPresence->Properties.Emplace(UpdatedProperty.Key, UpdatedProperty.Value);
