Added logic to print presence state via console command in both OSS and OnlineServices

diff --git a/Engine/Plugins/Online/OnlineServicesEOS/Source/Private/Online/PresenceEOS.cpp b/Engine/Plugins/Online/OnlineServicesEOS/Source/Private/Online/PresenceEOS.cpp
--- a/Engine/Plugins/Online/OnlineServicesEOS/Source/Private/Online/PresenceEOS.cpp
+++ b/Engine/Plugins/Online/OnlineServicesEOS/Source/Private/Online/PresenceEOS.cpp
@@ -117,6 +117,11 @@ void FPresenceEOS::UpdateConfig()
 	AsyncOpEnqueueDelay = Config.AsyncOpEnqueueDelay;
 }
 
+void FPresenceEOS::RegisterCommands()
+{
+	Super::RegisterCommands();
+}
+
 void FPresenceEOS::HandleAuthLoginStatusChanged(const FAuthLoginStatusChanged& EventParameters)
 {
 	if (EventParameters.LoginStatus == ELoginStatus::LoggedIn)
diff --git a/Engine/Plugins/Online/OnlineServicesEOS/Source/Public/Online/PresenceEOS.h b/Engine/Plugins/Online/OnlineServicesEOS/Source/Public/Online/PresenceEOS.h
--- a/Engine/Plugins/Online/OnlineServicesEOS/Source/Public/Online/PresenceEOS.h
+++ b/Engine/Plugins/Online/OnlineServicesEOS/Source/Public/Online/PresenceEOS.h
@@ -25,6 +25,7 @@ public:
 	ONLINESERVICESEOS_API virtual void Initialize() override;
 	ONLINESERVICESEOS_API virtual void PreShutdown() override;
 	ONLINESERVICESEOS_API virtual void UpdateConfig() override;
+	ONLINESERVICESEOS_API virtual void RegisterCommands() override;
 
 	ONLINESERVICESEOS_API virtual TOnlineAsyncOpHandle<FQueryPresence> QueryPresence(FQueryPresence::Params&& Params) override;
 	ONLINESERVICESEOS_API virtual TOnlineResult<FGetCachedPresence> GetCachedPresence(FGetCachedPresence::Params&& Params) override;
