Fix for user online status checks in V2 logic.
Refactored logic to read presence platform from EOSSDK.
Refactored how displaynames are retrieved in UserInfoEOS to prioritize passed in platforms.

#FailedValidations Preflight

diff --git a/Engine/Plugins/Online/OnlineFramework/Source/Party/Private/SocialToolkit.cpp b/Engine/Plugins/Online/OnlineFramework/Source/Party/Private/SocialToolkit.cpp
--- a/Engine/Plugins/Online/OnlineFramework/Source/Party/Private/SocialToolkit.cpp
+++ b/Engine/Plugins/Online/OnlineFramework/Source/Party/Private/SocialToolkit.cpp
@@ -341,18 +341,19 @@ void USocialToolkit::SetLocalUserOnlineStateV2(UE::Online::EUserPresenceStatus O
 		{
 			if (const FCommonAccountPtr LocalUserCommonAccountPtr = GetCommonAccountFromV1(CommonConfig, LocalUserNetId->AsShared(), GetFirstFrameworkInstanceName(LocalUserNetId)))
 			{
-				FAccountId LocalUserAccountId = LocalUserCommonAccountPtr->GetId(USocialManager::GetPresenceFrameworkInstance());
-
-				if (IPresencePtr Presence = OnlineServices->GetPresenceInterface())
+				if(FAccountId LocalUserAccountId = LocalUserCommonAccountPtr->GetId(USocialManager::GetPresenceFrameworkInstance()))
 				{
-					FPartialUpdatePresence::Params::FMutations Mutations;
-					Mutations.Status = OnlineState;
+					if (IPresencePtr Presence = OnlineServices->GetPresenceInterface())
+					{
+						FPartialUpdatePresence::Params::FMutations Mutations;
+						Mutations.Status = OnlineState;
 
-					FPartialUpdatePresence::Params Params = { LocalUserAccountId, Mutations };
+						FPartialUpdatePresence::Params Params = { LocalUserAccountId, Mutations };
 
-					UE_LOG(LogParty, Verbose, TEXT("Updating rich presence with OnlineServices"));
+						UE_LOG(LogParty, Verbose, TEXT("Updating rich presence with OnlineServices"));
 
-					Presence->PartialUpdatePresence(MoveTemp(Params));
+						Presence->PartialUpdatePresence(MoveTemp(Params));
+					}
 				}
 			}
 		}
@@ -1007,8 +1008,7 @@ void USocialToolkit::HandleServicesPresenceUpdated(const UE::Online::FPresenceUp
 	if (FCommonAccountPtr CommonAccountPtr = GetCommonAccountFromV2(CommonConfig, UpdatedPresence->AccountId, USocialManager::GetPresenceFrameworkInstance()))
 	{
 		// We need a services instance with an Auth adapter to make this conversion
-		const FName& PresenceFrameworkAdapterInstance = USocialManager::GetPresenceFrameworkAdapterInstance();
-		if (FUniqueNetIdPtr UserId = GetV1FromCommonAccount(CommonAccountPtr.ToSharedRef(), PresenceFrameworkAdapterInstance))
+		if (FUniqueNetIdPtr UserId = GetV1FromCommonAccount(CommonAccountPtr.ToSharedRef(), USocialManager::GetPresenceFrameworkAdapterInstance()))
 		{
 			QUICK_SCOPE_CYCLE_COUNTER(STAT_USocialToolkit_HandlePresenceReceived);
 			if (USocialUser* UpdatedUser = FindUser(UserId.ToSharedRef()))
diff --git a/Engine/Plugins/Online/OnlineFramework/Source/Party/Private/User/SocialUser.cpp b/Engine/Plugins/Online/OnlineFramework/Source/Party/Private/User/SocialUser.cpp
--- a/Engine/Plugins/Online/OnlineFramework/Source/Party/Private/User/SocialUser.cpp
+++ b/Engine/Plugins/Online/OnlineFramework/Source/Party/Private/User/SocialUser.cpp
@@ -1070,7 +1070,7 @@ FUserPlatform USocialUser::GetCurrentPlatform() const
 				const TSharedRef<const FUserPresence> CachedPresence = GetCachedPresenceResult.GetOkValue().Presence;
 
 				bIsPrimaryPresenceValid = true;
-				bIsPrimaryPresenceUserOnline = CachedPresence->Status == UE::Online::EUserPresenceStatus::Online;
+				bIsPrimaryPresenceUserOnline = CachedPresence->Status != UE::Online::EUserPresenceStatus::Offline;
 				bIsPrimariPresenceUserPlayingThisGame = CachedPresence->GameStatus == UE::Online::EUserPresenceGameStatus::PlayingThisGame;
 				PrimaryPresencePlatformStr = LexToString(CachedPresence->PlatformType);
 			}
diff --git a/Engine/Plugins/Online/OnlineFramework/Source/Party/Private/User/SocialUserList.cpp b/Engine/Plugins/Online/OnlineFramework/Source/Party/Private/User/SocialUserList.cpp
--- a/Engine/Plugins/Online/OnlineFramework/Source/Party/Private/User/SocialUserList.cpp
+++ b/Engine/Plugins/Online/OnlineFramework/Source/Party/Private/User/SocialUserList.cpp
@@ -423,7 +423,7 @@ bool FSocialUserList::EvaluateUserPresence(const USocialUser& User, ESocialSubsy
 			if (GetCachedPresenceResult.IsOk())
 			{
 				const TSharedRef<const FUserPresence> CachedPresence = GetCachedPresenceResult.GetOkValue().Presence;
-				bIsOnline = CachedPresence->Status == UE::Online::EUserPresenceStatus::Online;
+				bIsOnline = CachedPresence->Status != UE::Online::EUserPresenceStatus::Offline;
 				bIsPlayingThisGame = CachedPresence->GameStatus == UE::Online::EUserPresenceGameStatus::PlayingThisGame;
 			}
 		}
diff --git a/Engine/Plugins/Online/OnlineServicesEOS/Source/Private/Online/UserInfoEOS.cpp b/Engine/Plugins/Online/OnlineServicesEOS/Source/Private/Online/UserInfoEOS.cpp
--- a/Engine/Plugins/Online/OnlineServicesEOS/Source/Private/Online/UserInfoEOS.cpp
+++ b/Engine/Plugins/Online/OnlineServicesEOS/Source/Private/Online/UserInfoEOS.cpp
@@ -96,10 +96,15 @@ TOnlineAsyncOpHandle<FQueryUserInfo> FUserInfoEOS::QueryUserInfo(FQueryUserInfo:
 					Options.ExternalAccountIds = &TargetUserEasIdUtf8Str;
 					Options.ExternalAccountIdCount = 1;
 					EOS_HConnect ConnectHandle = EOS_Platform_GetConnectInterface(*GetServices<FOnlineServicesEOS>().GetEOSPlatformHandle());
+
+					GetServices<FOnlineServicesEpicCommon>().AddEOSSDKFastTick(Op);
+
 					EOS_Async(EOS_Connect_QueryExternalAccountMappings, ConnectHandle, Options, MoveTemp(Promise));
 				})
 				.Then([this](TOnlineAsyncOp<FQueryUserInfo>& Op, const EOS_Connect_QueryExternalAccountMappingsCallbackInfo* CallbackInfo) mutable
 				{
+					GetServices<FOnlineServicesEpicCommon>().RemoveEOSSDKFastTick(Op);
+
 					if (CallbackInfo->ResultCode != EOS_EResult::EOS_Success)
 					{
 						UE_LOG(LogOnlineServices, Warning, TEXT("EOS_Connect_QueryExternalAccountMappings failed with result=[%s]"), *LexToString(CallbackInfo->ResultCode));
@@ -132,12 +137,6 @@ TOnlineResult<FGetUserInfo> FUserInfoEOS::GetUserInfo(FGetUserInfo::Params&& Par
 		return TOnlineResult<FGetUserInfo>(Errors::InvalidParams());
 	}
 
-	EOS_UserInfo_CopyBestDisplayNameOptions Options = {};
-	Options.ApiVersion = 1;
-	UE_EOS_CHECK_API_MISMATCH(EOS_USERINFO_COPYBESTDISPLAYNAME_API_LATEST, 1);
-	Options.LocalUserId = GetEpicAccountIdChecked(Params.LocalAccountId);
-	Options.TargetUserId = TargetUserEasId;
-
 	EOS_UserInfo_BestDisplayName* EosBestDisplayName = nullptr;
 	ON_SCOPE_EXIT
 	{
@@ -147,9 +146,7 @@ TOnlineResult<FGetUserInfo> FUserInfoEOS::GetUserInfo(FGetUserInfo::Params&& Par
 		}
 	};
 
-	EOS_EResult EosResult = EOS_UserInfo_CopyBestDisplayName(UserInfoHandle, &Options, &EosBestDisplayName);
-
-	if (EosResult == EOS_EResult::EOS_UserInfo_BestDisplayNameIndeterminate)
+	auto CopyBestDisplayNameWithPlatform = [this, &Params, TargetUserEasId](EOS_UserInfo_BestDisplayName** EosBestDisplayName)
 	{
 		EOS_UserInfo_CopyBestDisplayNameWithPlatformOptions WithPlatformOptions = {};
 		WithPlatformOptions.ApiVersion = 1;
@@ -173,7 +170,28 @@ TOnlineResult<FGetUserInfo> FUserInfoEOS::GetUserInfo(FGetUserInfo::Params&& Par
 		
 		WithPlatformOptions.TargetPlatformType = OnlinePlatformType;
 
-		EosResult = EOS_UserInfo_CopyBestDisplayNameWithPlatform(UserInfoHandle, &WithPlatformOptions, &EosBestDisplayName);
+		return EOS_UserInfo_CopyBestDisplayNameWithPlatform(UserInfoHandle, &WithPlatformOptions, EosBestDisplayName);
+	};
+
+	EOS_EResult EosResult;
+	if (Params.DisplayNamePlatform.IsSet())
+	{
+		EosResult = CopyBestDisplayNameWithPlatform(&EosBestDisplayName);
+	}
+	else
+	{
+		EOS_UserInfo_CopyBestDisplayNameOptions Options = {};
+		Options.ApiVersion = 1;
+		UE_EOS_CHECK_API_MISMATCH(EOS_USERINFO_COPYBESTDISPLAYNAME_API_LATEST, 1);
+		Options.LocalUserId = GetEpicAccountIdChecked(Params.LocalAccountId);
+		Options.TargetUserId = TargetUserEasId;
+
+		EosResult = EOS_UserInfo_CopyBestDisplayName(UserInfoHandle, &Options, &EosBestDisplayName);
+
+		if (EosResult == EOS_EResult::EOS_UserInfo_BestDisplayNameIndeterminate)
+		{
+			EosResult = CopyBestDisplayNameWithPlatform(&EosBestDisplayName);
+		}
 	}
 
 	if(EosResult != EOS_EResult::EOS_Success)
