Added FastTick to EOS id resolution calls

diff --git a/Engine/Plugins/Online/OnlineServicesEOS/Source/Private/Online/AuthEOS.cpp b/Engine/Plugins/Online/OnlineServicesEOS/Source/Private/Online/AuthEOS.cpp
--- a/Engine/Plugins/Online/OnlineServicesEOS/Source/Private/Online/AuthEOS.cpp
+++ b/Engine/Plugins/Online/OnlineServicesEOS/Source/Private/Online/AuthEOS.cpp
@@ -769,8 +769,10 @@ TFuture<TArray<FAccountId>> FAuthEOS::ResolveAccountIds(const FAccountId& LocalA
 	Options.ExternalAccountIds = (const char**)EpicAccountIdStrPtrs.GetData();
 	Options.ExternalAccountIdCount = EpicAccountIdStrPtrs.Num();
 
+	TSharedPtr<IEOSFastTickLock> EOSSDKFastTickLock = GetServices<FOnlineServicesEpicCommon>().GetEOSPlatformHandle()->GetFastTickLock();
+
 	EOS_Async(EOS_Connect_QueryExternalAccountMappings, ConnectHandle, Options,
-	[this, WeakThis = AsWeak(), InEpicAccountIds, Promise = MoveTemp(Promise)](const EOS_Connect_QueryExternalAccountMappingsCallbackInfo* Data) mutable -> void
+	[this, WeakThis = AsWeak(), InEpicAccountIds, Promise = MoveTemp(Promise), EOSSDKFastTickLock = MoveTemp(EOSSDKFastTickLock)](const EOS_Connect_QueryExternalAccountMappingsCallbackInfo* Data) mutable -> void
 	{
 		TArray<FAccountId> AccountIds;
 		if (const TSharedPtr<IAuth> StrongThis = WeakThis.Pin())
@@ -853,9 +855,11 @@ TFuture<TArray<FAccountId>> FAuthEOS::ResolveAccountIds(const FAccountId& LocalA
 	Options.LocalUserId = IsRunningDedicatedServer() ? nullptr : GetProductUserIdChecked(LocalAccountId);
 	Options.ProductUserIds = MissingProductUserIds.GetData();
 	Options.ProductUserIdCount = MissingProductUserIds.Num();
-	
+
+	TSharedPtr<IEOSFastTickLock> EOSSDKFastTickLock = GetServices<FOnlineServicesEpicCommon>().GetEOSPlatformHandle()->GetFastTickLock();
+
 	EOS_Async(EOS_Connect_QueryProductUserIdMappings, ConnectHandle, Options,
-	[this, WeakThis = AsWeak(), InProductUserIds, Promise = MoveTemp(Promise)](const EOS_Connect_QueryProductUserIdMappingsCallbackInfo* Data) mutable -> void
+	[this, WeakThis = AsWeak(), InProductUserIds, Promise = MoveTemp(Promise), EOSSDKFastTickLock = MoveTemp(EOSSDKFastTickLock)](const EOS_Connect_QueryProductUserIdMappingsCallbackInfo* Data) mutable -> void
 	{
 		TArray<FAccountId> AccountIds;
 		if (const TSharedPtr<IAuth> StrongThis = WeakThis.Pin())
@@ -893,6 +897,7 @@ TFuture<TArray<FAccountId>> FAuthEOS::ResolveAccountIds(const FAccountId& LocalA
 				UE_LOG(LogOnlineServices, Warning, TEXT("ResolveAccountId failed to query external mapping Result=[%s]"), *LexToString(Data->ResultCode));
 			}
 		}
+
 		Promise.SetValue(MoveTemp(AccountIds));
 	});
 
