- Fix ExternalUIEOS not being instantiated in OnlineServicesEpicAccount
- Fix EOS_UI_ShowFriends errors always being reported as Unknown

#rnx

diff --git a/Engine/Plugins/Online/OnlineServicesEOS/Source/Private/Online/ExternalUIEOS.cpp b/Engine/Plugins/Online/OnlineServicesEOS/Source/Private/Online/ExternalUIEOS.cpp
--- a/Engine/Plugins/Online/OnlineServicesEOS/Source/Private/Online/ExternalUIEOS.cpp
+++ b/Engine/Plugins/Online/OnlineServicesEOS/Source/Private/Online/ExternalUIEOS.cpp
@@ -103,7 +103,7 @@ TOnlineAsyncOpHandle<FExternalUIShowFriendsUI> FExternalUIEOS::ShowFriendsUI(FEx
 		}
 		else
 		{
-			InAsyncOp.SetError(Errors::Unknown());
+			InAsyncOp.SetError(Errors::FromEOSResult(Data->ResultCode));
 		}
 	})
 	.Enqueue(GetSerialQueue());
