- Fix ExternalUIEOS not being instantiated in OnlineServicesEpicAccount
- Fix EOS_UI_ShowFriends errors always being reported as Unknown

#was-backout 45729538
#rnx

diff --git a/Engine/Plugins/Online/OnlineServicesEOS/Source/Private/Online/ExternalUIEOS.cpp b/Engine/Plugins/Online/OnlineServicesEOS/Source/Private/Online/ExternalUIEOS.cpp
--- a/Engine/Plugins/Online/OnlineServicesEOS/Source/Private/Online/ExternalUIEOS.cpp
+++ b/Engine/Plugins/Online/OnlineServicesEOS/Source/Private/Online/ExternalUIEOS.cpp
@@ -17,6 +17,8 @@ FExternalUIEOS::FExternalUIEOS(FOnlineServicesEpicCommon& InServices)
 {
 }
 
+FExternalUIEOS::~FExternalUIEOS() = default;
+
 void FExternalUIEOS::Initialize()
 {
 	Super::Initialize();
@@ -103,7 +105,7 @@ TOnlineAsyncOpHandle<FExternalUIShowFriendsUI> FExternalUIEOS::ShowFriendsUI(FEx
 		}
 		else
 		{
-			InAsyncOp.SetError(Errors::Unknown());
+			InAsyncOp.SetError(Errors::FromEOSResult(Data->ResultCode));
 		}
 	})
 	.Enqueue(GetSerialQueue());
diff --git a/Engine/Plugins/Online/OnlineServicesEOS/Source/Public/Online/ExternalUIEOS.h b/Engine/Plugins/Online/OnlineServicesEOS/Source/Public/Online/ExternalUIEOS.h
--- a/Engine/Plugins/Online/OnlineServicesEOS/Source/Public/Online/ExternalUIEOS.h
+++ b/Engine/Plugins/Online/OnlineServicesEOS/Source/Public/Online/ExternalUIEOS.h
@@ -25,8 +25,11 @@ public:
 	using Super = FExternalUIEOSGS;
 
 	ONLINESERVICESEOS_API FExternalUIEOS(FOnlineServicesEpicCommon& InOwningSubsystem);
+	ONLINESERVICESEOS_API virtual ~FExternalUIEOS();
+
 	ONLINESERVICESEOS_API virtual void Initialize() override;
 	ONLINESERVICESEOS_API virtual void UpdateConfig() override;
+
 	ONLINESERVICESEOS_API virtual TOnlineAsyncOpHandle<FExternalUIShowLoginUI> ShowLoginUI(FExternalUIShowLoginUI::Params&& Params) override;
 	ONLINESERVICESEOS_API virtual TOnlineAsyncOpHandle<FExternalUIShowFriendsUI> ShowFriendsUI(FExternalUIShowFriendsUI::Params&& Params) override;
 
diff --git a/Engine/Plugins/Online/OnlineServicesEOSGS/Source/Private/Online/ExternalUIEOSGS.cpp b/Engine/Plugins/Online/OnlineServicesEOSGS/Source/Private/Online/ExternalUIEOSGS.cpp
--- a/Engine/Plugins/Online/OnlineServicesEOSGS/Source/Private/Online/ExternalUIEOSGS.cpp
+++ b/Engine/Plugins/Online/OnlineServicesEOSGS/Source/Private/Online/ExternalUIEOSGS.cpp
@@ -14,6 +14,8 @@ FExternalUIEOSGS::FExternalUIEOSGS(FOnlineServicesEpicCommon& InServices)
 {
 }
 
+FExternalUIEOSGS::~FExternalUIEOSGS() = default;
+
 void FExternalUIEOSGS::Initialize()
 {
 	Super::Initialize();
diff --git a/Engine/Plugins/Online/OnlineServicesEOSGS/Source/Public/Online/ExternalUIEOSGS.h b/Engine/Plugins/Online/OnlineServicesEOSGS/Source/Public/Online/ExternalUIEOSGS.h
--- a/Engine/Plugins/Online/OnlineServicesEOSGS/Source/Public/Online/ExternalUIEOSGS.h
+++ b/Engine/Plugins/Online/OnlineServicesEOSGS/Source/Public/Online/ExternalUIEOSGS.h
@@ -40,6 +40,7 @@ public:
 	using Super = FExternalUICommon;
 
 	UE_API FExternalUIEOSGS(FOnlineServicesEpicCommon& InOwningSubsystem);
+	UE_API virtual ~FExternalUIEOSGS();
 
 	UE_API virtual void Initialize() override;
 	UE_API virtual void PreShutdown() override;
