Changes to display name related logic to get around integrated platform requirements when using EOSSDK

diff --git a/Engine/Plugins/Online/OnlineServicesEOS/Source/Private/Online/PresenceEOS.cpp b/Engine/Plugins/Online/OnlineServicesEOS/Source/Private/Online/PresenceEOS.cpp
--- a/Engine/Plugins/Online/OnlineServicesEOS/Source/Private/Online/PresenceEOS.cpp
+++ b/Engine/Plugins/Online/OnlineServicesEOS/Source/Private/Online/PresenceEOS.cpp
@@ -808,7 +808,7 @@ TSharedRef<FUserPresence> FPresenceEOS::FindOrCreatePresence(FAccountId LocalAcc
 	return UserPresence;
 }
 
-EOnlinePlatformType IntegratedPlatform_To_OnlinePlatformType(const FString& IntegratedPlatformStr)
+EOnlinePlatformType FPresenceEOS::IntegratedPlatform_To_OnlinePlatformType(const FString& IntegratedPlatformStr)
 {
 	if (IntegratedPlatformStr == TEXT("EGS"))
 	{
@@ -841,6 +841,24 @@ EOnlinePlatformType IntegratedPlatform_To_OnlinePlatformType(const FString& Inte
 	}
 }
 
+bool FPresenceEOS::ReadOnlinePlatformType(const EOS_Presence_Info* PresenceInfo, EOnlinePlatformType& OutOnlinePlatformType)
+{
+	// We set the PlatformType by translating the IntegratedPlatform field
+	FString IntegratedPlatformStr = UTF8_TO_TCHAR(PresenceInfo->IntegratedPlatform);
+	if (!IntegratedPlatformStr.IsEmpty() && IntegratedPlatformStr != TEXT("UNKNOWN")) // Integrated platform will show as UNKNOWN when it has not been set in platform builds, so we want to ignore it too
+	{
+		EOnlinePlatformType OnlinePlatformType = IntegratedPlatform_To_OnlinePlatformType(IntegratedPlatformStr);
+		if (OutOnlinePlatformType != OnlinePlatformType)
+		{
+			OutOnlinePlatformType = OnlinePlatformType;
+
+			return true;
+		}
+	}
+
+	return false;
+}
+
 void FPresenceEOS::UpdateUserPresence(FAccountId LocalAccountId, FAccountId PresenceAccountId)
 {
 	const EOS_EpicAccountId LocalUserId = GetEpicAccountId(LocalAccountId);
@@ -964,17 +982,7 @@ void FPresenceEOS::UpdateUserPresence(FAccountId LocalAccountId, FAccountId Pres
 			UserPresence->Properties.Reset();
 		}
 
-		// We set the PlatformType by translating the IntegratedPlatform field
-		FString IntegratedPlatformStr = UTF8_TO_TCHAR(PresenceInfo->IntegratedPlatform);
-		if (!IntegratedPlatformStr.IsEmpty())
-		{
-			EOnlinePlatformType OnlinePlatformType = IntegratedPlatform_To_OnlinePlatformType(IntegratedPlatformStr);
-			if (UserPresence->PlatformType != OnlinePlatformType)
-			{
-				bPresenceHasChanged = true;
-				UserPresence->PlatformType = OnlinePlatformType;
-			}
-		}
+		bPresenceHasChanged |= ReadOnlinePlatformType(PresenceInfo, UserPresence->PlatformType);
 
 		EOS_Presence_Info_Release(PresenceInfo);
 	}
diff --git a/Engine/Plugins/Online/OnlineServicesEOS/Source/Public/Online/PresenceEOS.h b/Engine/Plugins/Online/OnlineServicesEOS/Source/Public/Online/PresenceEOS.h
--- a/Engine/Plugins/Online/OnlineServicesEOS/Source/Public/Online/PresenceEOS.h
+++ b/Engine/Plugins/Online/OnlineServicesEOS/Source/Public/Online/PresenceEOS.h
@@ -49,6 +49,9 @@ protected:
 	/** Allow derived classes to modify the content of partial presence updates */
 	ONLINESERVICESEOS_API virtual void ModifyPartialPresenceUpdate(FPartialUpdatePresence::Params::FMutations& Mutations);
 
+	ONLINESERVICESEOS_API static EOnlinePlatformType IntegratedPlatform_To_OnlinePlatformType(const FString& IntegratedPlatformStr);
+	ONLINESERVICESEOS_API virtual bool ReadOnlinePlatformType(const EOS_Presence_Info* PresenceInfo, EOnlinePlatformType& OutOnlinePlatformType);
+
 	EOS_HPresence PresenceHandle = nullptr;
 
 	/** Login status changed event handle */
