Fix for Horde issue 1100655
Added missing virtual destructor declaration

diff --git a/Engine/Plugins/Online/OnlineServicesEOS/Source/Public/Online/PresenceEOS.h b/Engine/Plugins/Online/OnlineServicesEOS/Source/Public/Online/PresenceEOS.h
--- a/Engine/Plugins/Online/OnlineServicesEOS/Source/Public/Online/PresenceEOS.h
+++ b/Engine/Plugins/Online/OnlineServicesEOS/Source/Public/Online/PresenceEOS.h
@@ -21,6 +21,7 @@ public:
 	using Super = FPresenceCommon;
 
 	ONLINESERVICESEOS_API FPresenceEOS(FOnlineServicesEpicCommon& InServices);
+	virtual ~FPresenceEOS() = default;
 
 	ONLINESERVICESEOS_API virtual void Initialize() override;
 	ONLINESERVICESEOS_API virtual void PreShutdown() override;
