Refactored virtual destructors to avoid vtable errors

diff --git a/Engine/Plugins/Online/OnlineServicesEOS/Source/Private/Online/PresenceEOS.cpp b/Engine/Plugins/Online/OnlineServicesEOS/Source/Private/Online/PresenceEOS.cpp
--- a/Engine/Plugins/Online/OnlineServicesEOS/Source/Private/Online/PresenceEOS.cpp
+++ b/Engine/Plugins/Online/OnlineServicesEOS/Source/Private/Online/PresenceEOS.cpp
@@ -79,6 +79,8 @@ FPresenceEOS::FPresenceEOS(FOnlineServicesEpicCommon& InServices)
 {
 }
 
+FPresenceEOS::~FPresenceEOS() = default;
+
 void FPresenceEOS::Initialize()
 {
 	Super::Initialize();
diff --git a/Engine/Plugins/Online/OnlineServicesEOS/Source/Public/Online/PresenceEOS.h b/Engine/Plugins/Online/OnlineServicesEOS/Source/Public/Online/PresenceEOS.h
--- a/Engine/Plugins/Online/OnlineServicesEOS/Source/Public/Online/PresenceEOS.h
+++ b/Engine/Plugins/Online/OnlineServicesEOS/Source/Public/Online/PresenceEOS.h
@@ -21,7 +21,7 @@ public:
 	using Super = FPresenceCommon;
 
 	ONLINESERVICESEOS_API FPresenceEOS(FOnlineServicesEpicCommon& InServices);
-	ONLINESERVICESEOS_API virtual ~FPresenceEOS() = default;
+	ONLINESERVICESEOS_API virtual ~FPresenceEOS();
 
 	ONLINESERVICESEOS_API virtual void Initialize() override;
 	ONLINESERVICESEOS_API virtual void PreShutdown() override;
