Removal of temporary logs for eos presence testing

diff --git a/Engine/Plugins/Online/OnlineServicesEOS/Source/Private/Online/PresenceEOS.cpp b/Engine/Plugins/Online/OnlineServicesEOS/Source/Private/Online/PresenceEOS.cpp
--- a/Engine/Plugins/Online/OnlineServicesEOS/Source/Private/Online/PresenceEOS.cpp
+++ b/Engine/Plugins/Online/OnlineServicesEOS/Source/Private/Online/PresenceEOS.cpp
@@ -567,10 +567,6 @@ TOnlineAsyncOpHandle<FPartialUpdatePresence> FPresenceEOS::PartialUpdatePresence
 							Promise.EmplaceValue();
 							return;
 						}
-						else
-						{
-							UE_LOG(LogOnlineServices, Warning, TEXT("UpdatePresence: Status set to %u"), (uint8)Mutations.Status.GetValue()); // Temp logging
-						}
 					}
 
 					// Raw rich text
@@ -592,10 +588,6 @@ TOnlineAsyncOpHandle<FPartialUpdatePresence> FPresenceEOS::PartialUpdatePresence
 							Promise.EmplaceValue();
 							return;
 						}
-						else
-						{
-							UE_LOG(LogOnlineServices, Warning, TEXT("UpdatePresence: RichText set to %s"), *Mutations.StatusString.GetValue()); // Temp logging
-						}
 					}
 					// Removed fields
 					if (Mutations.RemovedProperties.Num() > 0)
@@ -612,8 +604,6 @@ TOnlineAsyncOpHandle<FPartialUpdatePresence> FPresenceEOS::PartialUpdatePresence
 							RecordId.ApiVersion = 1;
 							UE_EOS_CHECK_API_MISMATCH(EOS_PRESENCEMODIFICATION_DATARECORDID_API_LATEST, 1);
 							RecordId.Key = Utf8Key.Get();
-
-							UE_LOG(LogOnlineServices, Warning, TEXT("UpdatePresence: Removing field %s"), *RemovedProperty); // Temp logging
 						}
 
 						EOS_PresenceModification_DeleteDataOptions DataOptions = { };
@@ -695,6 +685,7 @@ TOnlineAsyncOpHandle<FPartialUpdatePresence> FPresenceEOS::PartialUpdatePresence
 					SetPresenceOptions.PresenceModificationHandle = ChangeHandle;
 					EOS_Async(EOS_Presence_SetPresence, PresenceHandle, SetPresenceOptions, MoveTemp(Promise));
 					EOS_PresenceModification_Release(ChangeHandle);
+
 					return;
 				}
 				else
@@ -707,8 +698,6 @@ TOnlineAsyncOpHandle<FPartialUpdatePresence> FPresenceEOS::PartialUpdatePresence
 			{
 				GetServices<FOnlineServicesEpicCommon>().RemoveEOSSDKFastTick(InAsyncOp);
 
-				UE_LOG(LogOnlineServices, Warning, TEXT("SetPresenceResult: [%s]"), *LexToString(Data->ResultCode));
-
 				if (Data->ResultCode == EOS_EResult::EOS_Success)
 				{
 					// Update local presence
