Changes to transmit rich presence string properly

diff --git a/Engine/Plugins/Online/OnlineServicesEOS/Source/Private/Online/PresenceEOS.cpp b/Engine/Plugins/Online/OnlineServicesEOS/Source/Private/Online/PresenceEOS.cpp
--- a/Engine/Plugins/Online/OnlineServicesEOS/Source/Private/Online/PresenceEOS.cpp
+++ b/Engine/Plugins/Online/OnlineServicesEOS/Source/Private/Online/PresenceEOS.cpp
@@ -358,12 +358,21 @@ TOnlineAsyncOpHandle<FUpdatePresence> FPresenceEOS::UpdatePresence(FUpdatePresen
 			}
 
 			// Raw rich text
-			// Convert the status string as the rich text string
+			FString RichPresenceString;
+			if (!Presence->RichPresenceString.IsEmpty())
+			{
+				RichPresenceString = Presence->RichPresenceString;
+			}
+			else if (!Presence->StatusString.IsEmpty()) // RichPresenceString is the intended way to set presence in UE, but we'll keep reading StatusString for backwards compatibility
+			{						
+				RichPresenceString = Presence->StatusString;
+			}
+
 			EOS_PresenceModification_SetRawRichTextOptions RawRichTextOptions = { };
 			RawRichTextOptions.ApiVersion = 1;
 			UE_EOS_CHECK_API_MISMATCH(EOS_PRESENCEMODIFICATION_SETRAWRICHTEXT_API_LATEST, 1);
 
-			const FTCHARToUTF8 Utf8RawRichText(*Presence->StatusString);
+			const FTCHARToUTF8 Utf8RawRichText(*RichPresenceString);
 			RawRichTextOptions.RichText = Utf8RawRichText.Get();
 
 			EOS_EResult SetRawRichTextResult = EOS_PresenceModification_SetRawRichText(ChangeHandle, &RawRichTextOptions);
@@ -570,14 +579,24 @@ TOnlineAsyncOpHandle<FPartialUpdatePresence> FPresenceEOS::PartialUpdatePresence
 					}
 
 					// Raw rich text
-					if (Mutations.StatusString.IsSet())
+					TOptional<FString> RichPresenceString;
+					if (Mutations.RichPresenceString.IsSet())
+					{
+						RichPresenceString = Mutations.RichPresenceString;
+					}
+					else if (Mutations.StatusString.IsSet()) // RichPresenceString is the intended way to set presence in UE, but we'll keep reading StatusString for backwards compatibility
+					{						
+						RichPresenceString = Mutations.StatusString;
+					}
+
+					if (RichPresenceString.IsSet())
 					{
 						// Convert the status string as the rich text string
 						EOS_PresenceModification_SetRawRichTextOptions RawRichTextOptions = { };
 						RawRichTextOptions.ApiVersion = 1;
 						UE_EOS_CHECK_API_MISMATCH(EOS_PRESENCEMODIFICATION_SETRAWRICHTEXT_API_LATEST, 1);
 
-						const FTCHARToUTF8 Utf8RawRichText(*Mutations.StatusString.GetValue());
+						const FTCHARToUTF8 Utf8RawRichText(*RichPresenceString.GetValue());
 						RawRichTextOptions.RichText = Utf8RawRichText.Get();
 
 						EOS_EResult SetRawRichTextResult = EOS_PresenceModification_SetRawRichText(ChangeHandle, &RawRichTextOptions);
@@ -589,6 +608,7 @@ TOnlineAsyncOpHandle<FPartialUpdatePresence> FPresenceEOS::PartialUpdatePresence
 							return;
 						}
 					}
+
 					// Removed fields
 					if (Mutations.RemovedProperties.Num() > 0)
 					{
@@ -728,6 +748,16 @@ TOnlineAsyncOpHandle<FPartialUpdatePresence> FPresenceEOS::PartialUpdatePresence
 						LocalUserPresence->StatusString = Mutations.StatusString.GetValue();
 					}
 
+					if (Mutations.RichPresenceString.IsSet())
+					{
+						if (LocalUserPresence->RichPresenceString != Mutations.RichPresenceString.GetValue())
+						{
+							bPresenceHasChanged = true;
+						}
+
+						LocalUserPresence->RichPresenceString = Mutations.RichPresenceString.GetValue();
+					}
+
 					if (Mutations.GameStatus.IsSet())
 					{
 						if (LocalUserPresence->GameStatus != Mutations.GameStatus.GetValue())
@@ -888,11 +918,14 @@ void FPresenceEOS::UpdateUserPresence(FAccountId LocalAccountId, FAccountId Pres
 			UserPresence->Status = NewPresenceState;
 		}
 
-		FString NewStatusString = UTF8_TO_TCHAR(PresenceInfo->RichText);
-		if (UserPresence->StatusString != NewStatusString)
+		FString NewRichPresenceString = UTF8_TO_TCHAR(PresenceInfo->RichText);
+		if (UserPresence->RichPresenceString != NewRichPresenceString)
 		{
 			bPresenceHasChanged = true;
-			UserPresence->StatusString = MoveTemp(NewStatusString);
+			UserPresence->RichPresenceString = MoveTemp(NewRichPresenceString);
+
+			 // RichPresenceString is the intended way to set presence in UE, but we'll keep populating StatusString for backwards compatibility
+			UserPresence->StatusString = UserPresence->RichPresenceString;
 		}
 
 		FString ProductIdString = UTF8_TO_TCHAR(PresenceInfo->ProductId);
