Added Platform parameter to EOS Presence to follow current functionality in Social layer
Miscellaneous fixes to PresenceEOS

diff --git a/Engine/Plugins/Online/OnlineFramework/Source/Party/Private/User/SocialUser.cpp b/Engine/Plugins/Online/OnlineFramework/Source/Party/Private/User/SocialUser.cpp
--- a/Engine/Plugins/Online/OnlineFramework/Source/Party/Private/User/SocialUser.cpp
+++ b/Engine/Plugins/Online/OnlineFramework/Source/Party/Private/User/SocialUser.cpp
@@ -1117,7 +1117,11 @@ FUserPlatform USocialUser::GetCurrentPlatform() const
 				bIsPrimaryPresenceValid = true;
 				bIsPrimaryPresenceUserOnline = CachedPresence->Status != UE::Online::EUserPresenceStatus::Offline;
 				bIsPrimariPresenceUserPlayingThisGame = CachedPresence->GameStatus == UE::Online::EUserPresenceGameStatus::PlayingThisGame;
-				PrimaryPresencePlatformStr = LexToString(CachedPresence->PlatformType);
+
+				if (const FPresenceProperty* PlatformAttribute = CachedPresence->Properties.Find(PLATFORM_ATTRIBUTE_KEY))
+				{
+					PrimaryPresencePlatformStr = PlatformAttribute->Get<FString>();
+				}
 			}
 		}
 		else
diff --git a/Engine/Plugins/Online/OnlineServices/Source/OnlineServicesInterface/Public/Online/Presence.h b/Engine/Plugins/Online/OnlineServices/Source/OnlineServicesInterface/Public/Online/Presence.h
--- a/Engine/Plugins/Online/OnlineServices/Source/OnlineServicesInterface/Public/Online/Presence.h
+++ b/Engine/Plugins/Online/OnlineServices/Source/OnlineServicesInterface/Public/Online/Presence.h
@@ -18,6 +18,7 @@ namespace UE::Online::Meta { template <typename StructType> struct TStructDetail
 
 #define SENT_TIME_ATTRIBUTE_KEY FString(TEXT("SentTimeAttributeKey"))
 #define SESSION_ID_ATTRIBUTE_KEY FString(TEXT("SessionIdAttributeKey"))
+#define PLATFORM_ATTRIBUTE_KEY FString(TEXT("PlatformAttributeKey"))
 
 namespace UE::Online {
 
diff --git a/Engine/Plugins/Online/OnlineServicesEOS/Source/Private/Online/PresenceEOS.cpp b/Engine/Plugins/Online/OnlineServicesEOS/Source/Private/Online/PresenceEOS.cpp
--- a/Engine/Plugins/Online/OnlineServicesEOS/Source/Private/Online/PresenceEOS.cpp
+++ b/Engine/Plugins/Online/OnlineServicesEOS/Source/Private/Online/PresenceEOS.cpp
@@ -288,12 +288,16 @@ TValueOrError<FPresenceProperty, UE::Online::FOnlineError> PresenceProperty_From
 		{
 			const char* RecordValueUtf8 = EOSPresenceString + 1;
 
+			OutRecordValueStr = UTF8_TO_TCHAR(RecordValueUtf8);
+
 			RecordValue.Set<int64>(FCStringUtf8::Atoi64((const UTF8CHAR*)RecordValueUtf8));
 		}
 		else if (EOSPresenceString[0] == 'd')
 		{
 			const char* RecordValueUtf8 = EOSPresenceString + 1;
 
+			OutRecordValueStr = UTF8_TO_TCHAR(RecordValueUtf8);
+
 			RecordValue.Set<double>(FCStringUtf8::Atod((const UTF8CHAR*)RecordValueUtf8));
 		}
 		else if (EOSPresenceString[0] == 's')
@@ -866,7 +870,7 @@ bool FPresenceEOS::ReadOnlinePlatformType(const EOS_Presence_Info* PresenceInfo,
 {
 	// We set the PlatformType by translating the IntegratedPlatform field
 	FString IntegratedPlatformStr = UTF8_TO_TCHAR(PresenceInfo->IntegratedPlatform);
-	if (!IntegratedPlatformStr.IsEmpty() && IntegratedPlatformStr != TEXT("UNKNOWN")) // Integrated platform will show as UNKNOWN when it has not been set in platform builds, so we want to ignore it too
+	if (!IntegratedPlatformStr.IsEmpty())
 	{
 		EOnlinePlatformType OnlinePlatformType = IntegratedPlatform_To_OnlinePlatformType(IntegratedPlatformStr);
 		if (OutOnlinePlatformType != OnlinePlatformType)
@@ -911,6 +915,7 @@ void FPresenceEOS::UpdateUserPresence(FAccountId LocalAccountId, FAccountId Pres
 	if (CopyPresenceResult == EOS_EResult::EOS_Success)
 	{
 		// Convert the presence data to our format
+
 		EUserPresenceStatus NewPresenceState = ToEPresenceState(PresenceInfo->Status);
 		if (UserPresence->Status != NewPresenceState)
 		{
@@ -928,15 +933,20 @@ void FPresenceEOS::UpdateUserPresence(FAccountId LocalAccountId, FAccountId Pres
 			UserPresence->StatusString = UserPresence->RichPresenceString;
 		}
 
-		FString ProductIdString = UTF8_TO_TCHAR(PresenceInfo->ProductId);
+		const FString ProductIdString = UTF8_TO_TCHAR(PresenceInfo->ProductId);
 		if (IEOSSDKManager* Manager = IEOSSDKManager::Get())
 		{
 			const FString& PlatformConfigName = StaticCast<FOnlineServicesEpicCommon*>(&Services)->GetEOSPlatformHandle()->GetConfigName();
 			if (const FEOSSDKPlatformConfig* Config = Manager->GetPlatformConfig(PlatformConfigName))
 			{
-				FString PlatformConfigProductIdString = Config->ProductId;
-				
-				UserPresence->GameStatus = ProductIdString == PlatformConfigProductIdString ? EUserPresenceGameStatus::PlayingThisGame : EUserPresenceGameStatus::PlayingOtherGame;
+				const FString PlatformConfigProductIdString = Config->ProductId;
+				EUserPresenceGameStatus NewGameStatus = ProductIdString == PlatformConfigProductIdString ? EUserPresenceGameStatus::PlayingThisGame : EUserPresenceGameStatus::PlayingOtherGame;
+
+				if (UserPresence->GameStatus != NewGameStatus)
+				{
+					bPresenceHasChanged = true;
+					UserPresence->GameStatus = MoveTemp(NewGameStatus);
+				}				
 			}
 		}
 
@@ -971,16 +981,25 @@ void FPresenceEOS::UpdateUserPresence(FAccountId LocalAccountId, FAccountId Pres
 								continue; // No change
 							}
 						}
+						else if (ExistingValue->IsType<bool>() && RecordValue.IsType<bool>() && ExistingValue->Get<bool>() == RecordValue.Get<bool>())
+						{
+							continue; // No change
+						}
 						else if ((ExistingValue->IsType<int64>() && RecordValue.IsType<int64>() && ExistingValue->Get<int64>() == RecordValue.Get<int64>()) ||
 								(ExistingValue->IsType<long long>() && RecordValue.IsType<long long>() && ExistingValue->Get<long long>() == RecordValue.Get<long long>())) // Setting a variant as int64 can appear as long long too
 						{
 							continue; // No change
 						}
-						else if (ExistingValue->IsType<FString>() && RecordValue.IsType<FString>() && ExistingValue->Get<FString>() == RecordValue.Get<FString>())
+						else if (ExistingValue->IsType<double>() && RecordValue.IsType<double>() && ExistingValue->Get<double>() == RecordValue.Get<double>())
 						{
 							continue; // No change
 						}
+						else if (ExistingValue->IsType<FString>() && RecordValue.IsType<FString>() && ExistingValue->Get<FString>() == RecordValue.Get<FString>())
+						{
+							continue; // No change
+						}						
 					}
+
 					bPresenceHasChanged = true;
 					UserPresence->Properties.Add(MoveTemp(RecordKey), MoveTemp(RecordValue));
 				}
@@ -1006,6 +1025,35 @@ void FPresenceEOS::UpdateUserPresence(FAccountId LocalAccountId, FAccountId Pres
 			UserPresence->Properties.Reset();
 		}
 
+		// We update the Platform attribute
+		if (PresenceInfo->Platform)
+		{
+			FString NewPlatformStr = UTF8_TO_TCHAR(PresenceInfo->Platform);
+			if (FPresenceProperty* PlatformAttribute = UserPresence->Properties.Find(PLATFORM_ATTRIBUTE_KEY))
+			{
+				const FString& PlatformAttributeStr = PlatformAttribute->Get<FString>();
+
+				if (PlatformAttributeStr != NewPlatformStr)
+				{
+					bPresenceHasChanged = true;
+					PlatformAttribute->Set<FString>(MoveTemp(NewPlatformStr));
+				}
+			}
+			else
+			{
+				bPresenceHasChanged = true;
+				UserPresence->Properties.AddVariant(PLATFORM_ATTRIBUTE_KEY, MoveTemp(NewPlatformStr));
+			}
+		}
+		else
+		{
+			if (FPresenceProperty* PlatformAttribute = UserPresence->Properties.Find(PLATFORM_ATTRIBUTE_KEY))
+			{
+				bPresenceHasChanged = true;
+				UserPresence->Properties.Remove(PLATFORM_ATTRIBUTE_KEY);
+			}
+		}
+
 		bPresenceHasChanged |= ReadOnlinePlatformType(PresenceInfo, UserPresence->PlatformType);
 
 		EOS_Presence_Info_Release(PresenceInfo);
