Removal of internal logic setting SentTime attribute.
Addition of engine logic setting SentTime attribute to reception time, until EOSSDK value is exposed.

diff --git a/Engine/Plugins/Online/OnlineServicesEOS/Source/Private/Online/PresenceEOS.cpp b/Engine/Plugins/Online/OnlineServicesEOS/Source/Private/Online/PresenceEOS.cpp
--- a/Engine/Plugins/Online/OnlineServicesEOS/Source/Private/Online/PresenceEOS.cpp
+++ b/Engine/Plugins/Online/OnlineServicesEOS/Source/Private/Online/PresenceEOS.cpp
@@ -1054,6 +1054,9 @@ void FPresenceEOS::UpdateUserPresence(FAccountId LocalAccountId, FAccountId Pres
 			}
 		}
 
+		// We set the SentTime attribute to the current time. Pending access to the actual sent time via EOSSDK
+		UserPresence->Properties.AddVariant<int64>(SENT_TIME_ATTRIBUTE_KEY, FDateTime::UtcNow().GetTicks());
+
 		bPresenceHasChanged |= ReadOnlinePlatformType(PresenceInfo, UserPresence->PlatformType);
 
 		EOS_Presence_Info_Release(PresenceInfo);
