Added logic to clear EOS presence state for user on auth logout

diff --git a/Engine/Plugins/Online/OnlineServicesEOS/Source/Private/Online/PresenceEOS.cpp b/Engine/Plugins/Online/OnlineServicesEOS/Source/Private/Online/PresenceEOS.cpp
--- a/Engine/Plugins/Online/OnlineServicesEOS/Source/Private/Online/PresenceEOS.cpp
+++ b/Engine/Plugins/Online/OnlineServicesEOS/Source/Private/Online/PresenceEOS.cpp
@@ -150,6 +150,10 @@ void FPresenceEOS::HandleAuthLoginStatusChanged(const FAuthLoginStatusChanged& E
 			PendingPresenceUpdates.Remove(EpicAccountId);
 		}
 	}
+	else if (EventParameters.LoginStatus == ELoginStatus::NotLoggedIn)
+	{
+		PresenceLists.Remove(EventParameters.AccountInfo->AccountId);
+	}
 }
 
 TOnlineAsyncOpHandle<FQueryPresence> FPresenceEOS::QueryPresence(FQueryPresence::Params&& InParams)
