Online Services EOS Presence Refactor:
- Optimistic cache
- Retry logic using EOSSDK cached state
- Echo suppression
- Added FUserPresence::operator==, used to detect differences and fire events.
- Simplified UpdateUserPresence

#Synced-CL 52537288

diff --git a/Engine/Plugins/Online/OnlineServices/Source/OnlineServicesInterface/Public/Online/Presence.h b/Engine/Plugins/Online/OnlineServices/Source/OnlineServicesInterface/Public/Online/Presence.h
--- a/Engine/Plugins/Online/OnlineServices/Source/OnlineServicesInterface/Public/Online/Presence.h
+++ b/Engine/Plugins/Online/OnlineServices/Source/OnlineServicesInterface/Public/Online/Presence.h
@@ -108,6 +108,8 @@ struct FUserPresence
 	EOnlinePlatformType PlatformType = EOnlinePlatformType::Unknown;
 	/** Session keys, and any platform specific presence properties */
 	FPresenceProperties Properties;
+
+	bool operator==(const FUserPresence& Other) const = default;
 };
 	
 struct FQueryPresence
diff --git a/Engine/Plugins/Online/OnlineServicesEOS/Source/Private/Online/PresenceEOS.cpp b/Engine/Plugins/Online/OnlineServicesEOS/Source/Private/Online/PresenceEOS.cpp
--- a/Engine/Plugins/Online/OnlineServicesEOS/Source/Private/Online/PresenceEOS.cpp
+++ b/Engine/Plugins/Online/OnlineServicesEOS/Source/Private/Online/PresenceEOS.cpp
@@ -18,11 +18,15 @@ namespace UE::Online {
 struct FPresenceEOSConfig
 {
 	double AsyncOpEnqueueDelay = 0.0;
+	double PresenceRetryDelaySec = 5.0;
+	int32 MaxPresenceRetries = 3;
 };
 
 namespace Meta {
 	BEGIN_ONLINE_STRUCT_META(FPresenceEOSConfig)
-		ONLINE_STRUCT_FIELD(FPresenceEOSConfig, AsyncOpEnqueueDelay)
+		ONLINE_STRUCT_FIELD(FPresenceEOSConfig, AsyncOpEnqueueDelay),
+		ONLINE_STRUCT_FIELD(FPresenceEOSConfig, PresenceRetryDelaySec),
+		ONLINE_STRUCT_FIELD(FPresenceEOSConfig, MaxPresenceRetries)
 	END_ONLINE_STRUCT_META()
 }
 
@@ -107,6 +111,14 @@ void FPresenceEOS::PreShutdown()
 	LoginStatusChangedHandle.Unbind();
 
 	OnPresenceChanged = nullptr;
+
+	if (RetryTickerHandle.IsValid())
+	{
+		FTSTicker::GetCoreTicker().RemoveTicker(RetryTickerHandle);
+		RetryTickerHandle.Reset();
+	}
+	
+	PendingRetries.Empty();
 }
 
 void FPresenceEOS::UpdateConfig()
@@ -117,6 +129,8 @@ void FPresenceEOS::UpdateConfig()
 	LoadConfig(Config);
 
 	AsyncOpEnqueueDelay = Config.AsyncOpEnqueueDelay;
+	PresenceRetryDelaySec = Config.PresenceRetryDelaySec;
+	MaxPresenceRetries = Config.MaxPresenceRetries;
 }
 
 void FPresenceEOS::RegisterCommands()
@@ -153,6 +167,7 @@ void FPresenceEOS::HandleAuthLoginStatusChanged(const FAuthLoginStatusChanged& E
 	else if (EventParameters.LoginStatus == ELoginStatus::NotLoggedIn)
 	{
 		PresenceLists.Remove(EventParameters.AccountInfo->AccountId);
+		ClearPresenceRetry(EventParameters.AccountInfo->AccountId);
 	}
 }
 
@@ -335,7 +350,30 @@ TOnlineAsyncOpHandle<FUpdatePresence> FPresenceEOS::UpdatePresence(FUpdatePresen
 		return Op->GetHandle();
 	}
 
-	Op->Then([this](TOnlineAsyncOp<FUpdatePresence>& InAsyncOp, TPromise<const EOS_Presence_SetPresenceCallbackInfo*>&& Promise) mutable
+	TSharedRef<FUserPresence> Presence = Params.Presence;
+	ModifyPresenceUpdate(Presence);
+
+	// Compute removed properties BEFORE applying optimistic cache update,
+	// since ApplyFullPresenceToCache will overwrite the cached presence with the new one.
+	// EOS needs to be specific on which fields are removed and which aren't, so grab the last presence and check which are deletions.
+	TArray<FString> RemovedProperties;
+	{
+		TSharedRef<FUserPresence> OldPresence = FindOrCreatePresence(Params.LocalAccountId, Params.LocalAccountId);
+		for (const TPair<FPresenceProperties::KeyType, FPresenceProperty>& Pair : OldPresence->Properties)
+		{
+			if (!Presence->Properties.Contains(Pair.Key))
+			{
+				RemovedProperties.Add(Pair.Key);
+			}
+		}
+	}
+
+	// Optimistic cache update: apply presence to local cache immediately before enqueue,
+	// so that GetCachedPresence returns up-to-date data even if multiple ops are queued.
+	CancelPresenceRetry(Params.LocalAccountId);
+	ApplyFullPresenceToCache(Params.LocalAccountId, *Presence);
+
+	Op->Then([this, RemovedProperties = MoveTemp(RemovedProperties)](TOnlineAsyncOp<FUpdatePresence>& InAsyncOp, TPromise<const EOS_Presence_SetPresenceCallbackInfo*>&& Promise) mutable
 	{
 		const FUpdatePresence::Params& Params = InAsyncOp.GetParams();
 		TSharedRef<FUserPresence> Presence = Params.Presence;
@@ -392,18 +430,6 @@ TOnlineAsyncOpHandle<FUpdatePresence> FPresenceEOS::UpdatePresence(FUpdatePresen
 				return;
 			}
 
-			// EOS needs to be specific on which fields are removed and which aren't, so grab the last presence and check which are deletions
-			TArray<FString> RemovedProperties;
-			TSharedRef<FUserPresence> LastUserPresence = FindOrCreatePresence(Params.LocalAccountId, Params.LocalAccountId);
-			for (const TPair<FPresenceProperties::KeyType, FPresenceProperty>& Pair : LastUserPresence->Properties)
-			{
-				if(!Presence->Properties.Contains(Pair.Key))
-				{
-					RemovedProperties.Add(Pair.Key);
-				}
-			}
-
-			// Removed fields
 			if (RemovedProperties.Num() > 0)
 			{
 				// EOS_PresenceModification_DeleteData
@@ -419,7 +445,7 @@ TOnlineAsyncOpHandle<FUpdatePresence> FPresenceEOS::UpdatePresence(FUpdatePresen
 					UE_EOS_CHECK_API_MISMATCH(EOS_PRESENCEMODIFICATION_DATARECORDID_API_LATEST, 1);
 					RecordId.Key = Utf8Key.Get();
 
-					UE_LOGF(LogOnlineServices, Warning, "UpdatePresence: Removing field %ls", *RemovedProperty); // Temp logging
+					UE_LOGF(LogOnlineServices, VeryVerbose, "UpdatePresence: Removing field %ls", *RemovedProperty); // Temp logging
 				}
 
 				EOS_PresenceModification_DeleteDataOptions DataOptions = { };
@@ -509,22 +535,18 @@ TOnlineAsyncOpHandle<FUpdatePresence> FPresenceEOS::UpdatePresence(FUpdatePresen
 	})
 	.Then([this](TOnlineAsyncOp<FUpdatePresence>& InAsyncOp, const EOS_Presence_SetPresenceCallbackInfo* Data) mutable
 	{
-		UE_LOGF(LogOnlineServices, Warning, "SetPresenceResult: [%ls]", *LexToString(Data->ResultCode));
-
 		if (Data->ResultCode == EOS_EResult::EOS_Success)
 		{
-			// Update local presence
-			const FUpdatePresence::Params& Params = InAsyncOp.GetParams();
-			TSharedRef<FUserPresence> Presence = Params.Presence;
-			ModifyPresenceUpdate(Presence);
-
-			TSharedRef<FUserPresence> LocalUserPresence = FindOrCreatePresence(Params.LocalAccountId, Params.LocalAccountId);
-			*LocalUserPresence = *Presence;
-
+			ClearPresenceRetry(InAsyncOp.GetParams().LocalAccountId);
 			InAsyncOp.SetResult(FUpdatePresence::Result());
 		}
 		else
 		{
+			const FUpdatePresence::Params& Params = InAsyncOp.GetParams();
+			
+			UE_LOGF(LogOnlineServices, Warning, "UpdatePresence: EOS_Presence_SetPresence failed with result [%ls] for user [%ls], scheduling retry", *LexToString(Data->ResultCode), *ToLogString(Params.LocalAccountId));
+			
+			SchedulePresenceRetry(Params.LocalAccountId);
 			InAsyncOp.SetError(Errors::FromEOSResult(Data->ResultCode));
 		}
 	})
@@ -552,6 +574,14 @@ TOnlineAsyncOpHandle<FPartialUpdatePresence> FPresenceEOS::PartialUpdatePresence
 			return Op->GetHandle();
 		}
 
+		FPartialUpdatePresence::Params::FMutations MutationsCopy = Params.Mutations;
+		ModifyPartialPresenceUpdate(MutationsCopy);
+
+		// Optimistic cache update: apply mutations to local cache immediately before enqueue,
+		// so that GetCachedPresence returns up-to-date data even if multiple ops are queued.		
+		CancelPresenceRetry(Params.LocalAccountId);
+		ApplyMutationsToCache(Params.LocalAccountId, MutationsCopy);
+
 		// Don't cache anything from Params as they could be modified by another merge in the meanwhile.
 		Op->Then([this](TOnlineAsyncOp<FPartialUpdatePresence>& InAsyncOp, TPromise<const EOS_Presence_SetPresenceCallbackInfo*>&& Promise) mutable
 			{
@@ -728,92 +758,16 @@ TOnlineAsyncOpHandle<FPartialUpdatePresence> FPresenceEOS::PartialUpdatePresence
 
 				if (Data->ResultCode == EOS_EResult::EOS_Success)
 				{
-					// Update local presence
-					bool bPresenceHasChanged = false;
-
-					const FPartialUpdatePresence::Params& Params = InAsyncOp.GetParams();
-					FPartialUpdatePresence::Params::FMutations Mutations = Params.Mutations;
-					ModifyPartialPresenceUpdate(Mutations);
-
-					TSharedRef<FUserPresence> LocalUserPresence = FindOrCreatePresence(Params.LocalAccountId, Params.LocalAccountId);
-					if (Mutations.Status.IsSet())
-					{
-						if (LocalUserPresence->Status != Mutations.Status.GetValue())
-						{
-							bPresenceHasChanged = true;
-						}
-
-						LocalUserPresence->Status = Mutations.Status.GetValue();
-					}
-
-					if (Mutations.StatusString.IsSet())
-					{
-						if (LocalUserPresence->StatusString != Mutations.StatusString.GetValue())
-						{
-							bPresenceHasChanged = true;
-						}
-
-						LocalUserPresence->StatusString = Mutations.StatusString.GetValue();
-					}
-
-					if (Mutations.RichPresenceString.IsSet())
-					{
-						if (LocalUserPresence->RichPresenceString != Mutations.RichPresenceString.GetValue())
-						{
-							bPresenceHasChanged = true;
-						}
-
-						LocalUserPresence->RichPresenceString = Mutations.RichPresenceString.GetValue();
-					}
-
-					if (Mutations.GameStatus.IsSet())
-					{
-						if (LocalUserPresence->GameStatus != Mutations.GameStatus.GetValue())
-						{
-							bPresenceHasChanged = true;
-						}
-
-						LocalUserPresence->GameStatus = Mutations.GameStatus.GetValue();
-					}
-
-					if (Mutations.Joinability.IsSet())
-					{
-						if (LocalUserPresence->Joinability != Mutations.Joinability.GetValue())
-						{
-							bPresenceHasChanged = true;
-						}
-
-						LocalUserPresence->Joinability = Mutations.Joinability.GetValue();
-					}
-
-					if (!Mutations.RemovedProperties.IsEmpty())
-					{
-						bPresenceHasChanged = true;
-					}
-					for (const FString& RemovedKey : Mutations.RemovedProperties)
-					{
-						LocalUserPresence->Properties.Remove(RemovedKey);
-					}
-
-					if (!Mutations.UpdatedProperties.IsEmpty())
-					{
-						bPresenceHasChanged = true;
-					}
-					for (const TPair<FPresenceProperties::KeyType, FPresenceProperty>& UpdatedProperty : Mutations.UpdatedProperties)
-					{
-						LocalUserPresence->Properties.Emplace(UpdatedProperty.Key, UpdatedProperty.Value);
-					}
-
-					if (bPresenceHasChanged)
-					{
-						FPresenceUpdated PresenceUpdatedParams = { Params.LocalAccountId, LocalUserPresence };
-						OnPresenceUpdatedEvent.Broadcast(PresenceUpdatedParams);
-					}
-
+					ClearPresenceRetry(InAsyncOp.GetParams().LocalAccountId);
 					InAsyncOp.SetResult(FPartialUpdatePresence::Result());
 				}
 				else
 				{
+					const FPartialUpdatePresence::Params& Params = InAsyncOp.GetParams();
+					
+					UE_LOGF(LogOnlineServices, Warning, "PartialUpdatePresence: EOS_Presence_SetPresence failed with result [%ls] for user [%ls], scheduling retry", *LexToString(Data->ResultCode), *ToLogString(Params.LocalAccountId));
+					
+					SchedulePresenceRetry(Params.LocalAccountId);
 					InAsyncOp.SetError(Errors::FromEOSResult(Data->ResultCode));
 				}
 			})
@@ -837,6 +791,184 @@ TSharedRef<FUserPresence> FPresenceEOS::FindOrCreatePresence(FAccountId LocalAcc
 	return UserPresence;
 }
 
+void FPresenceEOS::ApplyMutationsToCache(FAccountId LocalAccountId, const FPartialUpdatePresence::Params::FMutations& Mutations)
+{
+	TSharedRef<FUserPresence> LocalUserPresence = FindOrCreatePresence(LocalAccountId, LocalAccountId);
+	TSharedRef<FUserPresence> Updated = ApplyPresenceMutations(*LocalUserPresence, Mutations);
+
+	if (*LocalUserPresence != *Updated)
+	{
+		*LocalUserPresence = MoveTemp(*Updated);
+		FPresenceUpdated PresenceUpdatedParams = { LocalAccountId, LocalUserPresence };
+		OnPresenceUpdatedEvent.Broadcast(PresenceUpdatedParams);
+	}
+}
+
+void FPresenceEOS::ApplyFullPresenceToCache(FAccountId LocalAccountId, const FUserPresence& NewPresence)
+{
+	TSharedRef<FUserPresence> LocalUserPresence = FindOrCreatePresence(LocalAccountId, LocalAccountId);
+
+	if (*LocalUserPresence != NewPresence)
+	{
+		*LocalUserPresence = NewPresence;
+		FPresenceUpdated PresenceUpdatedParams = { LocalAccountId, LocalUserPresence };
+		OnPresenceUpdatedEvent.Broadcast(PresenceUpdatedParams);
+	}
+}
+
+void FPresenceEOS::SchedulePresenceRetry(FAccountId LocalAccountId)
+{
+	FPendingPresenceRetry& Retry = PendingRetries.FindOrAdd(LocalAccountId);
+	Retry.bPendingRetry = true;
+
+	if (!RetryTickerHandle.IsValid())
+	{
+		UE_LOGF(LogOnlineServices, Log, "PresenceEOS: Scheduling presence retry for user [%ls] in %.1f seconds", *ToLogString(LocalAccountId), PresenceRetryDelaySec);
+
+		RetryTickerHandle = FTSTicker::GetCoreTicker().AddTicker(
+			TEXT("PresenceEOSRetry"),
+			static_cast<float>(PresenceRetryDelaySec),
+			[WeakThis = AsWeak()](float)
+			{
+				if (TSharedPtr<FPresenceEOS> StrongThis = StaticCastSharedPtr<FPresenceEOS>(WeakThis.Pin()))
+				{
+					StrongThis->RetryTickerHandle.Reset();
+					StrongThis->StartPresenceRetry();
+				}
+				return false; // One-shot
+			});
+	}
+}
+
+void FPresenceEOS::StartPresenceRetry()
+{
+	// Snapshot the current set of users needing retries, don't remove them.
+	// Entries are only cleared when an EOS update succeeds or the user logs out.
+	TArray<FAccountId> UsersToRetry;
+	PendingRetries.GetKeys(UsersToRetry);
+
+	for (const FAccountId& LocalAccountId : UsersToRetry)
+	{
+		// Note: This pointer is invalidated if PendingRetries is structurally modified (e.g., ClearPresenceRetry on another user).
+		// Safe here because we only use it before any map removals for this iteration, and re-Find for each user.
+		FPendingPresenceRetry* Retry = PendingRetries.Find(LocalAccountId);
+		if (!Retry || !Retry->bPendingRetry)
+		{
+			continue;
+		}
+
+		Retry->RetryCount++;
+		if (MaxPresenceRetries >= 0 && Retry->RetryCount > MaxPresenceRetries)
+		{
+			UE_LOGF(LogOnlineServices, Warning, "PresenceEOS: Max retries (%d) exceeded for presence update, dropping. User=[%ls]", MaxPresenceRetries, *ToLogString(LocalAccountId));
+			ClearPresenceRetry(LocalAccountId);
+			continue;
+		}
+
+		UE_LOGF(LogOnlineServices, Log, "PresenceEOS: Retrying presence update for user [%ls] (attempt %d)", *ToLogString(LocalAccountId), Retry->RetryCount);
+
+		// Read our optimistic cache — this is the desired state we want EOS to match.
+		TOnlineResult<FGetCachedPresence> CachedPresenceResult = GetCachedPresence({ LocalAccountId, LocalAccountId });
+		if (!CachedPresenceResult.IsOk())
+		{
+			UE_LOGF(LogOnlineServices, Warning, "PresenceEOS: Could not get cached presence for retry, clearing. User=[%ls]", *ToLogString(LocalAccountId));
+			ClearPresenceRetry(LocalAccountId);
+			continue;
+		}
+
+		const TSharedRef<const FUserPresence>& CachedPresence = CachedPresenceResult.GetOkValue().Presence;
+
+		// Read the EOS SDK's own presence cache to determine what it currently has stored.
+		// Build a partial update (FMutations) representing the delta between EOS and our cache.
+		const EOS_EpicAccountId EpicAccountId = GetEpicAccountId(LocalAccountId);
+		if (EpicAccountId == nullptr)
+		{
+			UE_LOGF(LogOnlineServices, Warning, "PresenceEOS: No EpicAccountId for retry, clearing. User=[%ls]", *ToLogString(LocalAccountId));
+			ClearPresenceRetry(LocalAccountId);
+			continue;
+		}
+
+		EOS_Presence_Info* PresenceInfo = nullptr;
+		EOS_Presence_CopyPresenceOptions CopyOptions = { };
+		CopyOptions.ApiVersion = 3;
+		UE_EOS_CHECK_API_MISMATCH(EOS_PRESENCE_COPYPRESENCE_API_LATEST, 3);
+		CopyOptions.LocalUserId = EpicAccountId;
+		CopyOptions.TargetUserId = EpicAccountId;
+
+		EOS_EResult CopyResult = EOS_Presence_CopyPresence(PresenceHandle, &CopyOptions, &PresenceInfo);
+		if (CopyResult != EOS_EResult::EOS_Success)
+		{
+			UE_LOGF(LogOnlineServices, Warning, "PresenceEOS: CopyPresence failed during retry (result [%ls]), clearing. User=[%ls]", *LexToString(CopyResult), *ToLogString(LocalAccountId));
+			ClearPresenceRetry(LocalAccountId);
+			continue;
+		}
+
+		FPartialUpdatePresence::Params Params;
+		Params.LocalAccountId = LocalAccountId;
+		Params.Mutations = BuildPresenceDelta(PresenceInfo, *CachedPresence);
+
+		EOS_Presence_Info_Release(PresenceInfo);
+
+		PartialUpdatePresence(MoveTemp(Params));
+	}
+}
+
+FPartialUpdatePresence::Params::FMutations FPresenceEOS::BuildPresenceDelta(const EOS_Presence_Info* EOSPresenceInfo, const FUserPresence& DesiredPresence)
+{
+	FPartialUpdatePresence::Params::FMutations Mutations;
+
+	// Status
+	EUserPresenceStatus EOSStatus = ToEPresenceState(EOSPresenceInfo->Status);
+	if (EOSStatus != DesiredPresence.Status)
+	{
+		Mutations.Status = DesiredPresence.Status;
+	}
+
+	// Rich presence string (PartialUpdatePresence uses RichPresenceString with StatusString as fallback for SetRawRichText)
+	FString EOSRichText = UTF8_TO_TCHAR(EOSPresenceInfo->RichText);
+	if (EOSRichText != DesiredPresence.RichPresenceString)
+	{
+		Mutations.RichPresenceString = DesiredPresence.RichPresenceString;
+		Mutations.StatusString = DesiredPresence.StatusString;
+	}
+
+	// Properties: keys in EOS but not in desired state → removals
+	if (EOSPresenceInfo->Records)
+	{
+		TArrayView<const EOS_Presence_DataRecord> Records(EOSPresenceInfo->Records, EOSPresenceInfo->RecordsCount);
+		for (const EOS_Presence_DataRecord& Record : Records)
+		{
+			FString RecordKey = UTF8_TO_TCHAR(Record.Key);
+			if (!DesiredPresence.Properties.Contains(RecordKey))
+			{
+				Mutations.RemovedProperties.Add(MoveTemp(RecordKey));
+			}
+		}
+	}
+
+	// Properties: send all keys from the desired state, not just the ones that differ.
+	// A value-level diff would require parsing EOS record strings back into variants and
+	// comparing each type, which isn't worth the complexity for a retry path that fires rarely.
+	Mutations.UpdatedProperties = DesiredPresence.Properties;
+
+	return Mutations;
+}
+
+void FPresenceEOS::CancelPresenceRetry(FAccountId LocalAccountId)
+{
+	// Only clear bPendingRetry, not RetryCount.
+	// The retry path also calls PartialUpdatePresence, which calls this method. Resetting the count here would create an infinite retry loop.
+	if (FPendingPresenceRetry* Retry = PendingRetries.Find(LocalAccountId))
+	{
+		Retry->bPendingRetry = false;
+	}
+}
+
+void FPresenceEOS::ClearPresenceRetry(FAccountId LocalAccountId)
+{
+	PendingRetries.Remove(LocalAccountId);
+}
+
 EOnlinePlatformType FPresenceEOS::IntegratedPlatform_To_OnlinePlatformType(const FString& IntegratedPlatformStr)
 {
 	if (IntegratedPlatformStr == TEXT("EGS"))
@@ -904,9 +1036,14 @@ void FPresenceEOS::UpdateUserPresence(FAccountId LocalAccountId, FAccountId Pres
 		return;
 	}
 
-	bool bPresenceHasChanged = false;
 	TSharedRef<FUserPresence> UserPresence = FindOrCreatePresence(LocalAccountId, PresenceAccountId);
-	// TODO:  Handle updates for local users.  Don't want to conflict with UpdatePresence calls
+
+	// We remove the SentTime attribute before the comparison
+	UserPresence->Properties.Remove(SENT_TIME_ATTRIBUTE_KEY);
+
+	// Snapshot the current state so we can detect changes at the end via operator==.
+	const FUserPresence OldPresence = *UserPresence;
+	bool bPresenceChanged = false;
 
 	// Get presence from EOS
 	EOS_Presence_Info* PresenceInfo = nullptr;
@@ -919,23 +1056,13 @@ void FPresenceEOS::UpdateUserPresence(FAccountId LocalAccountId, FAccountId Pres
 	if (CopyPresenceResult == EOS_EResult::EOS_Success)
 	{
 		// Convert the presence data to our format
-
-		EUserPresenceStatus NewPresenceState = ToEPresenceState(PresenceInfo->Status);
-		if (UserPresence->Status != NewPresenceState)
-		{
-			bPresenceHasChanged = true;
-			UserPresence->Status = NewPresenceState;
-		}
+		UserPresence->Status = ToEPresenceState(PresenceInfo->Status);
 
 		FString NewRichPresenceString = UTF8_TO_TCHAR(PresenceInfo->RichText);
-		if (UserPresence->RichPresenceString != NewRichPresenceString)
-		{
-			bPresenceHasChanged = true;
-			UserPresence->RichPresenceString = MoveTemp(NewRichPresenceString);
-
-			 // RichPresenceString is the intended way to set presence in UE, but we'll keep populating StatusString for backwards compatibility
-			UserPresence->StatusString = UserPresence->RichPresenceString;
-		}
+		
+		// RichPresenceString is the intended way to set presence in UE, but we'll keep populating StatusString for backwards compatibility
+		UserPresence->StatusString = NewRichPresenceString;
+		UserPresence->RichPresenceString = MoveTemp(NewRichPresenceString);
 
 		const FString ProductIdString = UTF8_TO_TCHAR(PresenceInfo->ProductId);
 		if (IEOSSDKManager* Manager = IEOSSDKManager::Get())
@@ -944,13 +1071,7 @@ void FPresenceEOS::UpdateUserPresence(FAccountId LocalAccountId, FAccountId Pres
 			if (const FEOSSDKPlatformConfig* Config = Manager->GetPlatformConfig(PlatformConfigName))
 			{
 				const FString PlatformConfigProductIdString = Config->ProductId;
-				EUserPresenceGameStatus NewGameStatus = ProductIdString == PlatformConfigProductIdString ? EUserPresenceGameStatus::PlayingThisGame : EUserPresenceGameStatus::PlayingOtherGame;
-
-				if (UserPresence->GameStatus != NewGameStatus)
-				{
-					bPresenceHasChanged = true;
-					UserPresence->GameStatus = MoveTemp(NewGameStatus);
-				}				
+				UserPresence->GameStatus = ProductIdString == PlatformConfigProductIdString ? EUserPresenceGameStatus::PlayingThisGame : EUserPresenceGameStatus::PlayingOtherGame;
 			}
 		}
 
@@ -971,41 +1092,8 @@ void FPresenceEOS::UpdateUserPresence(FAccountId LocalAccountId, FAccountId Pres
 				TValueOrError<FPresenceProperty, FOnlineError> RecordValueResult = PresenceProperty_From_EOS_Presence_String(Record.Value, RecordValueStr);
 				if (RecordValueResult.HasValue())
 				{
-					FPresenceProperty RecordValue = RecordValueResult.GetValue();
-
 					RemovedKeys.Remove(RecordKey);
-					if (FPresenceProperties::ValueType* ExistingValue = UserPresence->Properties.Find(RecordKey))
-					{
-						if (ExistingValue->IsType<FPresencePropertiesRef>() && RecordValue.IsType<FPresencePropertiesRef>())
-						{
-							const FString ExistingValueJsonStr = NestedVariantToJson(ExistingValue->Get<FPresencePropertiesRef>().ToSharedPtr());
-
-							if (ExistingValueJsonStr == RecordValueStr)
-							{
-								continue; // No change
-							}
-						}
-						else if (ExistingValue->IsType<bool>() && RecordValue.IsType<bool>() && ExistingValue->Get<bool>() == RecordValue.Get<bool>())
-						{
-							continue; // No change
-						}
-						else if ((ExistingValue->IsType<int64>() && RecordValue.IsType<int64>() && ExistingValue->Get<int64>() == RecordValue.Get<int64>()) ||
-								(ExistingValue->IsType<long long>() && RecordValue.IsType<long long>() && ExistingValue->Get<long long>() == RecordValue.Get<long long>())) // Setting a variant as int64 can appear as long long too
-						{
-							continue; // No change
-						}
-						else if (ExistingValue->IsType<double>() && RecordValue.IsType<double>() && ExistingValue->Get<double>() == RecordValue.Get<double>())
-						{
-							continue; // No change
-						}
-						else if (ExistingValue->IsType<FString>() && RecordValue.IsType<FString>() && ExistingValue->Get<FString>() == RecordValue.Get<FString>())
-						{
-							continue; // No change
-						}						
-					}
-
-					bPresenceHasChanged = true;
-					UserPresence->Properties.Add(MoveTemp(RecordKey), MoveTemp(RecordValue));
+					UserPresence->Properties.Add(MoveTemp(RecordKey), RecordValueResult.GetValue());
 				}
 				else
 				{
@@ -1013,55 +1101,29 @@ void FPresenceEOS::UpdateUserPresence(FAccountId LocalAccountId, FAccountId Pres
 				}
 			}
 
-			// Any fields that have been removed
-			if (RemovedKeys.Num() > 0)
+			for (const FString& RemovedKey : RemovedKeys)
 			{
-				bPresenceHasChanged = true;
-				for (const FString& RemovedKey : RemovedKeys)
-				{
-					UserPresence->Properties.Remove(RemovedKey);
-				}
+				UserPresence->Properties.Remove(RemovedKey);
 			}
 		}
-		else if (UserPresence->Properties.Num() > 0)
+		else
 		{
-			bPresenceHasChanged = true;
 			UserPresence->Properties.Reset();
 		}
 
 		// We update the Platform attribute
 		if (PresenceInfo->Platform)
 		{
-			FString NewPlatformStr = UTF8_TO_TCHAR(PresenceInfo->Platform);
-			if (FPresenceProperty* PlatformAttribute = UserPresence->Properties.Find(PLATFORM_ATTRIBUTE_KEY))
-			{
-				const FString& PlatformAttributeStr = PlatformAttribute->Get<FString>();
-
-				if (PlatformAttributeStr != NewPlatformStr)
-				{
-					bPresenceHasChanged = true;
-					PlatformAttribute->Set<FString>(MoveTemp(NewPlatformStr));
-				}
-			}
-			else
-			{
-				bPresenceHasChanged = true;
-				UserPresence->Properties.AddVariant(PLATFORM_ATTRIBUTE_KEY, MoveTemp(NewPlatformStr));
-			}
+			UserPresence->Properties.AddVariant(PLATFORM_ATTRIBUTE_KEY, FString(UTF8_TO_TCHAR(PresenceInfo->Platform)));
 		}
 		else
 		{
-			if (FPresenceProperty* PlatformAttribute = UserPresence->Properties.Find(PLATFORM_ATTRIBUTE_KEY))
-			{
-				bPresenceHasChanged = true;
-				UserPresence->Properties.Remove(PLATFORM_ATTRIBUTE_KEY);
-			}
+			UserPresence->Properties.Remove(PLATFORM_ATTRIBUTE_KEY);
 		}
 
-		// We set the SentTime attribute to the current time. Pending access to the actual sent time via EOSSDK
-		UserPresence->Properties.AddVariant<int64>(SENT_TIME_ATTRIBUTE_KEY, FDateTime::UtcNow().GetTicks());
+		ReadOnlinePlatformType(PresenceInfo, UserPresence->PlatformType);
 
-		bPresenceHasChanged |= ReadOnlinePlatformType(PresenceInfo, UserPresence->PlatformType);
+		bPresenceChanged = *UserPresence != OldPresence;
 
 		EOS_Presence_Info_Release(PresenceInfo);
 	}
@@ -1070,7 +1132,11 @@ void FPresenceEOS::UpdateUserPresence(FAccountId LocalAccountId, FAccountId Pres
 		UE_LOGF(LogOnlineServices, Warning, "UpdateUserPresence: CopyPresence failed with result: %ls", *LexToString(CopyPresenceResult));
 	}
 
-	if (bPresenceHasChanged)
+	// We set the SentTime attribute to the current time. Pending access to the actual sent time via EOSSDK
+	// We set this after the comparison because it will always differ
+	UserPresence->Properties.AddVariant<int64>(SENT_TIME_ATTRIBUTE_KEY, FDateTime::UtcNow().GetTicks());
+
+	if (bPresenceChanged)
 	{
 		FPresenceUpdated PresenceUpdatedParams = { LocalAccountId, UserPresence };
 		OnPresenceUpdatedEvent.Broadcast(PresenceUpdatedParams);
@@ -1086,8 +1152,17 @@ void FPresenceEOS::HandlePresenceChanged(const EOS_Presence_PresenceChangedCallb
 			{
 				if (TSharedPtr<FPresenceEOS> StrongThis = StaticCastSharedPtr<FPresenceEOS>(WeakThis.Pin()))
 				{
-					UE_LOGF(LogOnlineServices, Verbose, "%s LocalAccountId=[%ls] PresenceAccountId=[%ls]", Func, *ToLogString(LocalAccountId), *ToLogString(PresenceAccountId));
-					StrongThis->UpdateUserPresence(LocalAccountId, PresenceAccountId);
+					const bool bIsLocalUser = (PresenceAccountId == LocalAccountId);
+					
+					UE_LOGF(LogOnlineServices, Verbose, "%s LocalAccountId=[%ls] PresenceAccountId=[%ls]%s", Func, *ToLogString(LocalAccountId), *ToLogString(PresenceAccountId), bIsLocalUser ? " (skipping echo)" : "");
+
+					// Skip processing EOS notifications for the local user's own presence.
+					// The local cache is updated optimistically in PartialUpdatePresence/UpdatePresence,
+					// and the EOS echo-back notification could overwrite newer optimistic state.
+					if (!bIsLocalUser)
+					{
+						StrongThis->UpdateUserPresence(LocalAccountId, PresenceAccountId);
+					}
 				}
 			});
 	}
diff --git a/Engine/Plugins/Online/OnlineServicesEOS/Source/Public/Online/PresenceEOS.h b/Engine/Plugins/Online/OnlineServicesEOS/Source/Public/Online/PresenceEOS.h
--- a/Engine/Plugins/Online/OnlineServicesEOS/Source/Public/Online/PresenceEOS.h
+++ b/Engine/Plugins/Online/OnlineServicesEOS/Source/Public/Online/PresenceEOS.h
@@ -2,6 +2,7 @@
 
 #pragma once
 
+#include "Containers/Ticker.h"
 #include "EOSSharedTypes.h"
 #include "Online/OnlineServicesEOSGSTypes.h"
 #include "Online/PresenceCommon.h"
@@ -64,6 +65,47 @@ protected:
 	FEOSEventRegistrationPtr OnPresenceChanged;
 
 	double AsyncOpEnqueueDelay = 0.0;
+
+	// Optimistic cache + Presence Update Retries
+
+	/**
+	 * Apply partial presence mutations to the local cache and broadcast OnPresenceUpdatedEvent if there was a change.
+	 */
+	void ApplyMutationsToCache(FAccountId LocalAccountId, const FPartialUpdatePresence::Params::FMutations& Mutations);
+
+	/**
+	 * Apply a full presence replacement to the local cache and broadcast OnPresenceUpdatedEvent if there was a change.
+	 */
+	void ApplyFullPresenceToCache(FAccountId LocalAccountId, const FUserPresence& NewPresence);
+
+	/** Build a FMutations representing the delta between the EOS SDK's presence state and a desired presence. */
+	FPartialUpdatePresence::Params::FMutations BuildPresenceDelta(const EOS_Presence_Info* EOSPresenceInfo, const FUserPresence& DesiredPresence);
+
+	/** Mark a user's presence as needing a retry push to the EOS backend. */
+	void SchedulePresenceRetry(FAccountId LocalAccountId);
+
+	/** Execute pending presence retries by pushing the current cache state to EOS. */
+	void StartPresenceRetry();
+
+	/** Cancel the next scheduled retry for a user without resetting the retry counter. */
+	void CancelPresenceRetry(FAccountId LocalAccountId);
+
+	/** Fully remove pending retry state for a user. Called on success, max retries exceeded, or logout. */
+	void ClearPresenceRetry(FAccountId LocalAccountId);
+
+	/** Pending retry state for a user whose presence failed to push to EOS. */
+	struct FPendingPresenceRetry
+	{
+		int32 RetryCount = 0;
+
+		/** Whether this entry should fire on the next ticker tick. Set to false when a presence update method is called; RetryCount is preserved so any later failure in the same operation continues counting. */
+		bool bPendingRetry = true;
+	};
+
+	TMap<FAccountId, FPendingPresenceRetry> PendingRetries;
+	FTSTicker::FDelegateHandle RetryTickerHandle;
+	double PresenceRetryDelaySec = 5.0;
+	int32 MaxPresenceRetries = 3;
 };
 
 /* UE::Online */ }
